Serialization: Initial revision of the UStruct serialization framework

diff --git a/Engine/Source/Runtime/Serialization/Private/Backends/JsonStructDeserializerBackend.cpp b/Engine/Source/Runtime/Serialization/Private/Backends/JsonStructDeserializerBackend.cpp
new file mode 100644
--- /dev/null
+++ b/Engine/Source/Runtime/Serialization/Private/Backends/JsonStructDeserializerBackend.cpp
@@ -0,0 +1,258 @@
+// Copyright 1998-2013 Epic Games, Inc. All Rights Reserved.
+
+#include "SerializationPrivatePCH.h"
+#include "Json.h"
+#include "JsonStructDeserializerBackend.h"
+
+
+/* Internal helpers
+ *****************************************************************************/
+
+namespace JsonStructDeserializerBackend
+{
+	// Clears the value of the given property.
+	void ClearPropertyValue( UProperty* Property, UProperty* Outer, void* Data, int32 ArrayIndex )
+	{
+		UArrayProperty* ArrayProperty = Cast<UArrayProperty>(Outer);
+
+		if (ArrayProperty != nullptr)
+		{
+			check(ArrayProperty->Inner == Property)
+
+			FScriptArrayHelper ArrayHelper(ArrayProperty, ArrayProperty->template ContainerPtrToValuePtr<void>(Data));
+			ArrayIndex = ArrayHelper.AddValue();
+		}
+
+		Property->ClearValue_InContainer(Data, ArrayIndex);
+	}
+
+	// Sets the value of the given property.
+	template<typename UPropertyType, typename PropertyType>
+	void SetPropertyValue( UProperty* Property, UProperty* Outer, void* Data, int32 ArrayIndex, const PropertyType& Value )
+	{
+		PropertyType* ValuePtr = nullptr;
+		UArrayProperty* ArrayProperty = Cast<UArrayProperty>(Outer);
+
+		if (ArrayProperty != nullptr)
+		{
+			check(ArrayProperty->Inner == Property)
+
+			FScriptArrayHelper ArrayHelper(ArrayProperty, ArrayProperty->template ContainerPtrToValuePtr<void>(Data));
+			int32 Index = ArrayHelper.AddValue();
+		
+			ValuePtr = (PropertyType*)ArrayHelper.GetRawPtr(Index);
+		}
+		else
+		{
+			UPropertyType* TypedProperty = Cast<UPropertyType>(Property);
+			check(TypedProperty)
+
+			ValuePtr = TypedProperty->template ContainerPtrToValuePtr<PropertyType>(Data, ArrayIndex);
+		}
+
+		*ValuePtr = Value;
+	}
+}
+
+
+/* IStructDeserializerBackend interface
+ *****************************************************************************/
+
+const FString& FJsonStructDeserializerBackend::GetCurrentPropertyName( ) const
+{
+	return JsonReader->GetIdentifier();
+}
+
+
+FString FJsonStructDeserializerBackend::GetDebugString( ) const
+{
+	return FString::Printf(TEXT("Line: %u, Ch: %u"), JsonReader->GetLineNumber(), JsonReader->GetCharacterNumber());
+}
+
+
+const FString& FJsonStructDeserializerBackend::GetLastErrorMessage( ) const
+{
+	return JsonReader->GetErrorMessage();
+}
+
+
+bool FJsonStructDeserializerBackend::GetNextToken( EStructDeserializerBackendTokens& OutToken )
+{
+	if (!JsonReader->ReadNext(LastNotation))
+	{
+		return false;
+	}
+
+	switch (LastNotation)
+	{
+	case EJsonNotation::ArrayEnd:
+		OutToken = EStructDeserializerBackendTokens::ArrayEnd;
+		break;
+
+	case EJsonNotation::ArrayStart:
+		OutToken = EStructDeserializerBackendTokens::ArrayStart;
+		break;
+
+	case EJsonNotation::Boolean:
+	case EJsonNotation::Null:
+	case EJsonNotation::Number:
+	case EJsonNotation::String:
+		{
+			OutToken = EStructDeserializerBackendTokens::Property;
+		}
+		break;
+
+	case EJsonNotation::Error:
+		OutToken = EStructDeserializerBackendTokens::Error;
+		break;
+
+	case EJsonNotation::ObjectEnd:
+		OutToken = EStructDeserializerBackendTokens::StructureEnd;
+		break;
+
+	case EJsonNotation::ObjectStart:
+		OutToken = EStructDeserializerBackendTokens::StructureStart;
+		break;
+
+	default:
+		OutToken = EStructDeserializerBackendTokens::None;
+	}
+
+	return true;
+}
+
+
+bool FJsonStructDeserializerBackend::ReadProperty( UProperty* Property, UProperty* Outer, void* Data, int32 ArrayIndex )
+{
+	using namespace JsonStructDeserializerBackend;
+
+	switch (LastNotation)
+	{
+	// boolean values
+	case EJsonNotation::Boolean:
+		{
+			bool BoolValue = JsonReader->GetValueAsBoolean();
+
+			if (Property->GetClass() == UBoolProperty::StaticClass())
+			{
+				SetPropertyValue<UBoolProperty, bool>(Property, Outer, Data, ArrayIndex, BoolValue);
+			}
+			else
+			{
+				UE_LOG(LogSerialization, Verbose, TEXT("Boolean field %s with value '%s' is not supported in UProperty type %s (%s)"), *Property->GetFName().ToString(), BoolValue ? *(GTrue.ToString()) : *(GFalse.ToString()), *Property->GetClass()->GetName(), *GetDebugString());
+
+				return false;
+			}
+		}
+		break;
+
+	// numeric values
+	case EJsonNotation::Number:
+		{
+			float NumericValue = JsonReader->GetValueAsNumber();
+
+			if (Property->GetClass() == UByteProperty::StaticClass())
+			{
+				SetPropertyValue<UByteProperty, int8>(Property, Outer, Data, ArrayIndex, (int8)NumericValue);
+			}
+			else if (Property->GetClass() == UDoubleProperty::StaticClass())
+			{
+				SetPropertyValue<UDoubleProperty, double>(Property, Outer, Data, ArrayIndex, NumericValue);
+			}
+			else if (Property->GetClass() == UFloatProperty::StaticClass())
+			{
+				SetPropertyValue<UFloatProperty, float>(Property, Outer, Data, ArrayIndex, NumericValue);
+			}
+			else if (Property->GetClass() == UIntProperty::StaticClass())
+			{
+				SetPropertyValue<UIntProperty, int32>(Property, Outer, Data, ArrayIndex, (int32)NumericValue);
+			}
+			else if (Property->GetClass() == UUInt32Property::StaticClass())
+			{
+				SetPropertyValue<UUInt32Property, uint32>(Property, Outer, Data, ArrayIndex, (uint32)NumericValue);
+			}
+			else if (Property->GetClass() == UInt16Property::StaticClass())
+			{
+				SetPropertyValue<UInt16Property, int16>(Property, Outer, Data, ArrayIndex, (int16)NumericValue);
+			}
+			else if (Property->GetClass() == UUInt16Property::StaticClass())
+			{
+				SetPropertyValue<UUInt16Property, uint16>(Property, Outer, Data, ArrayIndex, (uint16)NumericValue);
+			}
+			else if (Property->GetClass() == UInt64Property::StaticClass())
+			{
+				SetPropertyValue<UInt64Property, int64>(Property, Outer, Data, ArrayIndex, (int64)NumericValue);
+			}
+			else if (Property->GetClass() == UUInt64Property::StaticClass())
+			{
+				SetPropertyValue<UUInt64Property, uint64>(Property, Outer, Data, ArrayIndex, (uint64)NumericValue);
+			}
+			else if (Property->GetClass() == UInt8Property::StaticClass())
+			{
+				SetPropertyValue<UInt8Property, int8>(Property, Outer, Data, ArrayIndex, (int8)NumericValue);
+			}
+			else
+			{
+				UE_LOG(LogSerialization, Verbose, TEXT("Numeric field %s with value '%f' is not supported in UProperty type %s (%s)"), *Property->GetFName().ToString(), NumericValue, *Property->GetClass()->GetName(), *GetDebugString());
+
+				return false;
+			}
+		}
+		break;
+
+	// null values
+	case EJsonNotation::Null:
+		ClearPropertyValue(Property, Outer, Data, ArrayIndex);
+		break;
+
+	// strings, names & enumerations
+	case EJsonNotation::String:
+		{
+			const FString& StringValue = JsonReader->GetValueAsString();
+
+			if (Property->GetClass() == UStrProperty::StaticClass())
+			{
+				SetPropertyValue<UStrProperty, FString>(Property, Outer, Data, ArrayIndex, StringValue);
+			}
+			else if (Property->GetClass() == UNameProperty::StaticClass())
+			{
+				SetPropertyValue<UNameProperty, FName>(Property, Outer, Data, ArrayIndex, *StringValue);
+			}
+			else if (Property->GetClass() == UByteProperty::StaticClass())
+			{
+				UByteProperty* ByteProperty = Cast<UByteProperty>(Property);
+				int32 Index = ByteProperty->Enum->FindEnumIndex(*StringValue);
+
+				if (Index != INDEX_NONE)
+				{
+					SetPropertyValue<UByteProperty, uint8>(Property, Outer, Data, ArrayIndex, (uint8)Index);
+				}
+			}
+			else if (Property->GetClass() == UClassProperty::StaticClass())
+			{
+				SetPropertyValue<UClassProperty, UClass*>(Property, Outer, Data, ArrayIndex, LoadObject<UClass>(NULL, *StringValue, NULL, LOAD_NoWarn));
+			}
+			else
+			{
+				UE_LOG(LogSerialization, Verbose, TEXT("String field %s with value '%s' is not supported in UProperty type %s (%s)"), *Property->GetFName().ToString(), *StringValue, *Property->GetClass()->GetName(), *GetDebugString());
+
+				return false;
+			}
+		}
+		break;
+	}
+
+	return true;
+}
+
+
+void FJsonStructDeserializerBackend::SkipArray( ) 
+{
+	JsonReader->SkipArray();
+}
+
+
+void FJsonStructDeserializerBackend::SkipStructure( )
+{
+	JsonReader->SkipObject();
+}
diff --git a/Engine/Source/Runtime/Serialization/Private/Backends/JsonStructSerializerBackend.cpp b/Engine/Source/Runtime/Serialization/Private/Backends/JsonStructSerializerBackend.cpp
new file mode 100644
--- /dev/null
+++ b/Engine/Source/Runtime/Serialization/Private/Backends/JsonStructSerializerBackend.cpp
@@ -0,0 +1,190 @@
+// Copyright 1998-2013 Epic Games, Inc. All Rights Reserved.
+
+#include "SerializationPrivatePCH.h"
+#include "Json.h"
+#include "JsonStructSerializerBackend.h"
+
+
+/* Internal helpers
+ *****************************************************************************/
+
+namespace JsonStructSerializerBackend
+{
+	// Writes a property value to the serialization output.
+	template<typename ValueType>
+	void WritePropertyValue( const TSharedRef<TJsonWriter<TCHAR, TPrettyJsonPrintPolicy<TCHAR>>> JsonWriter, UProperty* Property, const ValueType& Value )
+	{
+		if ((Property == nullptr) || (Property->ArrayDim > 1) || (Property->GetOuter()->GetClass() == UArrayProperty::StaticClass()))
+		{
+			JsonWriter->WriteValue(Value);
+		}
+		else
+		{
+			JsonWriter->WriteValue(Property->GetName(), Value);
+		}
+	}
+
+	// Writes a null value to the serialization output.
+	void WriteNull( const TSharedRef<TJsonWriter<TCHAR, TPrettyJsonPrintPolicy<TCHAR>>> JsonWriter, UProperty* Property )
+	{
+		if ((Property == nullptr) || (Property->ArrayDim > 1) || (Property->GetOuter()->GetClass() == UArrayProperty::StaticClass()))
+		{
+			JsonWriter->WriteNull();
+		}
+		else
+		{
+			JsonWriter->WriteNull(Property->GetName());
+		}
+	}
+}
+
+
+/* IStructSerializerBackend interface
+ *****************************************************************************/
+
+void FJsonStructSerializerBackend::BeginArray( UProperty* Property )
+{
+	UObject* Outer = Property->GetOuter();
+
+	if ((Outer != nullptr) && (Outer->GetClass() == UArrayProperty::StaticClass()))
+	{
+		JsonWriter->WriteArrayStart();
+	}
+	else
+	{
+		JsonWriter->WriteArrayStart(Property->GetName());
+	}
+}
+
+
+void FJsonStructSerializerBackend::BeginStructure( UProperty* Property )
+{
+	UObject* Outer = Property->GetOuter();
+
+	if ((Outer != nullptr) && (Outer->GetClass() == UArrayProperty::StaticClass()))
+	{
+		JsonWriter->WriteObjectStart();
+	}
+	else
+	{
+		JsonWriter->WriteObjectStart(Property->GetName());
+	}
+}
+
+
+void FJsonStructSerializerBackend::BeginStructure( UStruct* TypeInfo )
+{
+	JsonWriter->WriteObjectStart();
+}
+
+
+void FJsonStructSerializerBackend::EndArray( UProperty* Property )
+{
+	JsonWriter->WriteArrayEnd();
+}
+
+
+void FJsonStructSerializerBackend::EndStructure( )
+{
+	JsonWriter->WriteObjectEnd();
+}
+
+
+void FJsonStructSerializerBackend::WriteComment( const FString& Comment )
+{
+	// Json does not support comments
+}
+
+
+void FJsonStructSerializerBackend::WriteProperty( UProperty* Property, const void* Data, UStruct* TypeInfo, int32 ArrayIndex )
+{
+	using namespace JsonStructSerializerBackend;
+
+	// booleans
+	if (TypeInfo == UBoolProperty::StaticClass())
+	{
+		WritePropertyValue(JsonWriter, Property, Cast<UBoolProperty>(Property)->GetPropertyValue_InContainer(Data, ArrayIndex));
+	}
+
+	// unsigned bytes & enumerations
+	else if (TypeInfo == UByteProperty::StaticClass())
+	{
+		UByteProperty* ByteProperty = Cast<UByteProperty>(Property);
+
+		if (ByteProperty->IsEnum())
+		{
+			WritePropertyValue(JsonWriter, Property, ByteProperty->Enum->GetEnumName(ByteProperty->GetPropertyValue_InContainer(Data, ArrayIndex)));
+		}
+		else
+		{
+			WritePropertyValue(JsonWriter, Property, (float)Cast<UByteProperty>(Property)->GetPropertyValue_InContainer(Data, ArrayIndex));
+		}			
+	}
+
+	// floating point numbers
+	else if (TypeInfo == UDoubleProperty::StaticClass())
+	{
+		WritePropertyValue(JsonWriter, Property, Cast<UDoubleProperty>(Property)->GetPropertyValue_InContainer(Data, ArrayIndex));
+	}
+	else if (TypeInfo == UFloatProperty::StaticClass())
+	{
+		WritePropertyValue(JsonWriter, Property, Cast<UFloatProperty>(Property)->GetPropertyValue_InContainer(Data, ArrayIndex));
+	}
+
+	// signed integers
+	else if (TypeInfo == UIntProperty::StaticClass())
+	{
+		WritePropertyValue(JsonWriter, Property, (float)Cast<UIntProperty>(Property)->GetPropertyValue_InContainer(Data, ArrayIndex));
+	}
+	else if (TypeInfo == UInt8Property::StaticClass())
+	{
+		WritePropertyValue(JsonWriter, Property, (float)Cast<UInt8Property>(Property)->GetPropertyValue_InContainer(Data, ArrayIndex));
+	}
+	else if (TypeInfo == UInt16Property::StaticClass())
+	{
+		WritePropertyValue(JsonWriter, Property, (float)Cast<UInt16Property>(Property)->GetPropertyValue_InContainer(Data, ArrayIndex));
+	}
+	else if (TypeInfo == UInt64Property::StaticClass())
+	{
+		WritePropertyValue(JsonWriter, Property, (float)Cast<UInt64Property>(Property)->GetPropertyValue_InContainer(Data, ArrayIndex));
+	}
+
+	// unsigned integers
+	else if (TypeInfo == UUInt16Property::StaticClass())
+	{
+		WritePropertyValue(JsonWriter, Property, (float)Cast<UUInt16Property>(Property)->GetPropertyValue_InContainer(Data, ArrayIndex));
+	}
+	else if (TypeInfo == UUInt32Property::StaticClass())
+	{
+		WritePropertyValue(JsonWriter, Property, (float)Cast<UUInt32Property>(Property)->GetPropertyValue_InContainer(Data, ArrayIndex));
+	}
+	else if (TypeInfo == UUInt64Property::StaticClass())
+	{
+		WritePropertyValue(JsonWriter, Property, (float)Cast<UUInt64Property>(Property)->GetPropertyValue_InContainer(Data, ArrayIndex));
+	}
+
+	// names & strings
+	else if (TypeInfo == UNameProperty::StaticClass())
+	{
+		WritePropertyValue(JsonWriter, Property, Cast<UNameProperty>(Property)->GetPropertyValue_InContainer(Data, ArrayIndex).ToString());
+	}
+	else if (TypeInfo == UStrProperty::StaticClass())
+	{
+		WritePropertyValue(JsonWriter, Property, Cast<UStrProperty>(Property)->GetPropertyValue_InContainer(Data, ArrayIndex));
+	}
+
+	// classes & objects
+	else if (TypeInfo == UClassProperty::StaticClass())
+	{
+		WritePropertyValue(JsonWriter, Property, Cast<UClassProperty>(Property)->GetPropertyValue_InContainer(Data, ArrayIndex)->GetPathName());
+	}
+	else if (TypeInfo == UObjectProperty::StaticClass())
+	{
+		WriteNull(JsonWriter, Property);
+	}
+	
+	else
+	{
+		UE_LOG(LogSerialization, Verbose, TEXT("FJsonStructSerializerBackend: Property %s cannot be serialized, because its type (%s) is not supported"), *Property->GetFName().ToString(), *TypeInfo->GetFName().ToString());
+	}
+}
diff --git a/Engine/Source/Runtime/Serialization/Private/SerializationModule.cpp b/Engine/Source/Runtime/Serialization/Private/SerializationModule.cpp
new file mode 100644
--- /dev/null
+++ b/Engine/Source/Runtime/Serialization/Private/SerializationModule.cpp
@@ -0,0 +1,30 @@
+// Copyright 1998-2013 Epic Games, Inc. All Rights Reserved.
+
+#include "SerializationPrivatePCH.h"
+#include "ModuleManager.h"
+
+
+//DEFINE_LOG_CATEGORY(LogSerialization);
+
+
+/**
+ * Implements the Serialization module.
+ */
+class FSerializationModule
+	: public IModuleInterface
+{
+public:
+
+	// IModuleInterface interface
+
+	virtual void StartupModule( ) override { }
+	virtual void ShutdownModule( ) override { }
+
+	virtual bool SupportsDynamicReloading( ) override
+	{
+		return true;
+	}
+};
+
+
+IMPLEMENT_MODULE(FSerializationModule, Serialization);
diff --git a/Engine/Source/Runtime/Serialization/Private/SerializationPrivatePCH.h b/Engine/Source/Runtime/Serialization/Private/SerializationPrivatePCH.h
new file mode 100644
--- /dev/null
+++ b/Engine/Source/Runtime/Serialization/Private/SerializationPrivatePCH.h
@@ -0,0 +1,10 @@
+// Copyright 1998-2013 Epic Games, Inc. All Rights Reserved.
+
+#pragma once
+
+
+/* Private dependencies
+ *****************************************************************************/
+
+#include "Core.h"
+#include "CoreUObject.h"
diff --git a/Engine/Source/Runtime/Serialization/Private/StructDeserializer.cpp b/Engine/Source/Runtime/Serialization/Private/StructDeserializer.cpp
new file mode 100644
--- /dev/null
+++ b/Engine/Source/Runtime/Serialization/Private/StructDeserializer.cpp
@@ -0,0 +1,305 @@
+// Copyright 1998-2013 Epic Games, Inc. All Rights Reserved.
+
+#include "SerializationPrivatePCH.h"
+#include "IStructDeserializerBackend.h"
+#include "StructDeserializer.h"
+
+
+/* Internal helpers
+ *****************************************************************************/
+
+namespace StructDeserializer
+{
+	// Structure for the read state stack.
+	struct FReadState
+	{
+		// Holds the property's current array index.
+		int32 ArrayIndex;
+
+		// Holds a pointer to the property's data.
+		void* Data;
+
+		// Holds the property's meta data.
+		UProperty* Property;
+
+		// Holds a pointer to the UStruct describing the data.
+		UStruct* TypeInfo;
+	};
+
+
+	// Finds the class for the given stack state.
+	UStruct* FindClass( const FReadState& State )
+	{
+		UStruct* Class = nullptr;
+
+		if (State.Property != nullptr)
+		{
+			UProperty* ParentProperty = State.Property;
+			UArrayProperty* ArrayProperty = Cast<UArrayProperty>(ParentProperty);
+
+			if (ArrayProperty != nullptr)
+			{
+				ParentProperty = ArrayProperty->Inner;
+			}
+
+			UStructProperty* StructProperty = Cast<UStructProperty>(ParentProperty);
+			UObjectPropertyBase* ObjectProperty = Cast<UObjectPropertyBase>(ParentProperty);
+
+			if (StructProperty != nullptr)
+			{
+				Class = StructProperty->Struct;
+			}
+			else if (ObjectProperty != nullptr)
+			{
+				Class = ObjectProperty->PropertyClass;
+			}
+		}
+		else
+		{
+			UObject* RootObject = static_cast<UObject*>(State.Data);
+			Class = RootObject->GetClass();
+		}
+
+		return Class;
+	}
+}
+
+
+/* FStructDeserializer static interface
+ *****************************************************************************/
+
+bool FStructDeserializer::Deserialize( void* OutStruct, UStruct& TypeInfo, IStructDeserializerBackend& Backend, const FStructDeserializerPolicies& Policies )
+{
+	using namespace StructDeserializer;
+
+	check(OutStruct != nullptr);
+
+	// initialize deserialization
+	FReadState CurrentState;
+	{
+		CurrentState.ArrayIndex = 0;
+		CurrentState.Data = OutStruct;
+		CurrentState.Property = nullptr;
+		CurrentState.TypeInfo = &TypeInfo;
+	}
+
+	TArray<FReadState> StateStack;
+	EStructDeserializerBackendTokens Token;
+
+	// process state stack
+	while (Backend.GetNextToken(Token))
+	{
+		FString PropertyName = Backend.GetCurrentPropertyName();
+
+		switch (Token)
+		{
+		case EStructDeserializerBackendTokens::ArrayEnd:
+			{
+				if (StateStack.Num() == 0)
+				{
+					UE_LOG(LogSerialization, Verbose, TEXT("Malformed input: Found ArrayEnd without matching ArrayStart"));
+
+					return false;
+				}
+
+				CurrentState = StateStack.Pop();
+			}
+			break;
+
+		case EStructDeserializerBackendTokens::ArrayStart:
+			{
+				FReadState NewState;
+
+				NewState.Property = FindField<UProperty>(CurrentState.TypeInfo, *PropertyName);
+
+				if (NewState.Property == nullptr)
+				{
+					if (Policies.MissingFields != EStructDeserializerErrorPolicies::Ignore)
+					{
+						UE_LOG(LogSerialization, Verbose, TEXT("The array property '%s' does not exist"), *PropertyName);
+					}
+
+					if (Policies.MissingFields == EStructDeserializerErrorPolicies::Error)
+					{
+						return false;
+					}
+
+					Backend.SkipArray();
+				}
+				else
+				{
+					NewState.ArrayIndex = 0;
+					NewState.Data = CurrentState.Data;
+					NewState.TypeInfo = FindClass(NewState);
+
+					StateStack.Push(CurrentState);
+					CurrentState = NewState;
+				}
+			}
+			break;
+
+		case EStructDeserializerBackendTokens::Error:
+			{
+				return false;
+			}
+
+		case EStructDeserializerBackendTokens::Property:
+			{
+				UProperty* Property = nullptr;
+
+				if (PropertyName.IsEmpty())
+				{
+					// handle array property
+					UArrayProperty* ArrayProperty = Cast<UArrayProperty>(CurrentState.Property);
+
+					if (ArrayProperty != nullptr)
+					{
+						// dynamic array property
+						Property = ArrayProperty->Inner;
+					}
+					else
+					{
+						// static array property
+						Property = CurrentState.Property;
+					}
+
+					if (Property == nullptr)
+					{
+						// error: no property for array element
+						if (Policies.MissingFields != EStructDeserializerErrorPolicies::Ignore)
+						{
+							UE_LOG(LogSerialization, Verbose, TEXT("Failed to serialize array element %i"), CurrentState.ArrayIndex);
+						}
+
+						return false;
+					}
+					else if (!Backend.ReadProperty(Property, CurrentState.Property, CurrentState.Data, CurrentState.ArrayIndex))
+					{
+						UE_LOG(LogSerialization, Verbose, TEXT("The array element '%s[%i]' could not be read (%s)"), *PropertyName, CurrentState.ArrayIndex, *Backend.GetDebugString());
+					}
+
+					++CurrentState.ArrayIndex;
+				}
+				else
+				{
+					// handle scalar property
+					Property = FindField<UProperty>(CurrentState.TypeInfo, *PropertyName);
+
+					if (Property == nullptr)
+					{
+						// error: property not found
+						if (Policies.MissingFields != EStructDeserializerErrorPolicies::Ignore)
+						{
+							UE_LOG(LogSerialization, Verbose, TEXT("The property '%s' does not exist"), *PropertyName);
+						}
+
+						if (Policies.MissingFields == EStructDeserializerErrorPolicies::Error)
+						{
+							return false;
+						}
+					}
+					else if (!Backend.ReadProperty(Property, CurrentState.Property, CurrentState.Data, CurrentState.ArrayIndex))
+					{
+						UE_LOG(LogSerialization, Verbose, TEXT("The property '%s' could not be read (%s)"), *PropertyName, *Backend.GetDebugString());
+					}
+				}
+			}
+			break;
+
+		case EStructDeserializerBackendTokens::StructureEnd:
+			{
+				if (StateStack.Num() == 0)
+				{
+					return true;
+				}
+
+				CurrentState = StateStack.Pop();
+			}
+			break;
+
+		case EStructDeserializerBackendTokens::StructureStart:
+			{
+				FReadState NewState;
+
+				if (PropertyName.IsEmpty())
+				{
+					// skip root structure
+					if (CurrentState.Property == nullptr)
+					{
+						continue;
+					}
+
+					// handle structured array element
+					UArrayProperty* ArrayProperty = Cast<UArrayProperty>(CurrentState.Property);
+
+					if (ArrayProperty == nullptr)
+					{
+						UE_LOG(LogSerialization, Verbose, TEXT("Property %s is not an array"));
+
+						return false;
+					}
+
+					FScriptArrayHelper ArrayHelper(ArrayProperty, ArrayProperty->ContainerPtrToValuePtr<void>(CurrentState.Data));
+					const int32 ArrayIndex = ArrayHelper.AddValue();
+
+					NewState.Property = ArrayProperty->Inner;
+					NewState.Data = ArrayHelper.GetRawPtr(ArrayIndex);
+				}
+				else
+				{
+					// handle structured property
+					NewState.Property = FindField<UProperty>(CurrentState.TypeInfo, *PropertyName);
+
+					if (NewState.Property == nullptr)
+					{
+						// error: structured property not found
+						if (Policies.MissingFields != EStructDeserializerErrorPolicies::Ignore)
+						{
+							UE_LOG(LogSerialization, Verbose, TEXT("Structured property '%s' not found"), *PropertyName);
+						}
+
+						if (Policies.MissingFields == EStructDeserializerErrorPolicies::Error)
+						{
+							return false;
+						}
+					}
+					else
+					{
+						NewState.Data = NewState.Property->ContainerPtrToValuePtr<void>(CurrentState.Data);
+					}
+				}
+
+				if (NewState.Property != nullptr)
+				{
+					NewState.ArrayIndex = 0;
+					NewState.TypeInfo = FindClass(NewState);
+
+					StateStack.Push(CurrentState);
+					CurrentState = NewState;
+				}
+				else
+				{
+					// error: structured property not found
+					Backend.SkipStructure();
+
+					if (Policies.MissingFields != EStructDeserializerErrorPolicies::Ignore)
+					{
+						UE_LOG(LogSerialization, Verbose, TEXT("Structured property '%s' not found"), *PropertyName);
+					}
+
+					if (Policies.MissingFields == EStructDeserializerErrorPolicies::Error)
+					{
+						return false;
+					}
+				}
+			}
+
+		default:
+
+			continue;
+		}
+	}
+
+	// root structure not completed
+	return false;
+}
diff --git a/Engine/Source/Runtime/Serialization/Private/StructSerializer.cpp b/Engine/Source/Runtime/Serialization/Private/StructSerializer.cpp
new file mode 100644
--- /dev/null
+++ b/Engine/Source/Runtime/Serialization/Private/StructSerializer.cpp
@@ -0,0 +1,211 @@
+// Copyright 1998-2013 Epic Games, Inc. All Rights Reserved.
+
+#include "SerializationPrivatePCH.h"
+#include "IStructSerializerBackend.h"
+#include "StructSerializer.h"
+
+
+/* Internal helpers
+ *****************************************************************************/
+
+namespace StructSerializer
+{
+	// Structure for the write state stack.
+	struct FWriteState
+	{
+		// Holds a pointer to the property's data.
+		const void* Data;
+
+		// Holds a flag indicating whether the property has been processed.
+		bool HasBeenProcessed;
+
+		// Holds the property's meta data.
+		UProperty* Property;
+
+		// Holds a pointer to the UStruct describing the data.
+		UStruct* TypeInfo;
+	};
+
+
+	// Gets the value from the given property.
+	template<typename UPropertyType, typename PropertyType>
+	PropertyType* GetPropertyValue( const FWriteState& State, UProperty* Property )
+	{
+		PropertyType* ValuePtr = nullptr;
+		UArrayProperty* ArrayProperty = Cast<UArrayProperty>(State.Property);
+
+		if (ArrayProperty)
+		{
+			check(ArrayProperty->Inner == Property);
+
+			FScriptArrayHelper ArrayHelper(ArrayProperty, ArrayProperty->template ContainerPtrToValuePtr<void>(State.Data));
+			int32 Index = ArrayHelper.AddValue();
+		
+			ValuePtr = (PropertyType*)ArrayHelper.GetRawPtr( Index );
+		}
+		else
+		{
+			UPropertyType* TypedProperty = Cast<UPropertyType>(Property);
+			check(TypedProperty != nullptr);
+
+			ValuePtr = TypedProperty->template ContainerPtrToValuePtr<PropertyType>(State.Data);
+		}
+
+		return ValuePtr;
+	}
+}
+
+
+/* FStructSerializer static interface
+ *****************************************************************************/
+
+void FStructSerializer::Serialize( const void* Struct, UStruct& TypeInfo, IStructSerializerBackend& Backend, const FStructSerializerPolicies& Policies )
+{
+	using namespace StructSerializer;
+
+	check(Struct != nullptr);
+
+	// initialize serialization
+	TArray<FWriteState> StateStack;
+	{
+		FWriteState NewState;
+
+		NewState.Data = Struct;
+		NewState.Property = nullptr;
+		NewState.TypeInfo = &TypeInfo;
+		NewState.HasBeenProcessed = false;
+
+		StateStack.Push(NewState);
+	}
+
+	// process state stack
+	while (StateStack.Num() > 0)
+	{
+		FWriteState CurrentState = StateStack.Pop();
+
+		// structures
+		if ((CurrentState.Property == nullptr) || (CurrentState.TypeInfo == UStructProperty::StaticClass()))
+		{
+			if (!CurrentState.HasBeenProcessed)
+			{
+				const void* NewData = CurrentState.Data;
+
+				// write object start
+				if (CurrentState.Property == nullptr)
+				{
+					Backend.BeginStructure(CurrentState.TypeInfo);
+				}
+				else
+				{
+					UObject* Outer = CurrentState.Property->GetOuter();
+
+					if ((Outer == nullptr) || (Outer->GetClass() != UArrayProperty::StaticClass()))
+					{
+						NewData = CurrentState.Property->ContainerPtrToValuePtr<void>(CurrentState.Data);
+					}
+
+					Backend.BeginStructure(CurrentState.Property);
+				}
+
+				CurrentState.HasBeenProcessed = true;
+				StateStack.Push(CurrentState);
+
+				// serialize fields
+				if (CurrentState.Property != nullptr)
+				{
+					UStructProperty* StructProperty = Cast<UStructProperty>(CurrentState.Property);
+
+					if (StructProperty != nullptr)
+					{
+						CurrentState.TypeInfo = StructProperty->Struct;
+					}
+					else
+					{
+						UObjectPropertyBase* ObjectProperty = Cast<UObjectPropertyBase>(CurrentState.Property);
+
+						if (ObjectProperty != nullptr)
+						{
+							CurrentState.TypeInfo = ObjectProperty->PropertyClass;
+						}
+					}
+				}
+
+				TArray<FWriteState> NewStates;
+
+				for (TFieldIterator<UProperty> It(CurrentState.TypeInfo, EFieldIteratorFlags::IncludeSuper); It; ++It)
+				{
+					FWriteState NewState;
+
+					NewState.Data = NewData;
+					NewState.Property = *It;
+					NewState.TypeInfo = It->GetClass();
+					NewState.HasBeenProcessed = false;
+
+					NewStates.Add(NewState);
+				}
+
+				// push child properties on stack (in reverse order)
+				for (int Index = NewStates.Num() - 1; Index >= 0; --Index)
+				{
+					StateStack.Push(NewStates[Index]);
+				}
+			}
+			else
+			{
+				Backend.EndStructure();
+			}
+		}
+
+		// dynamic arrays
+		else if (CurrentState.TypeInfo == UArrayProperty::StaticClass())
+		{
+			if (!CurrentState.HasBeenProcessed)
+			{
+				Backend.BeginArray(CurrentState.Property);
+
+				CurrentState.HasBeenProcessed = true;
+				StateStack.Push(CurrentState);
+
+				UArrayProperty* ArrayProperty = Cast<UArrayProperty>(CurrentState.Property);
+				FScriptArrayHelper ArrayHelper(ArrayProperty, ArrayProperty->ContainerPtrToValuePtr<void>(CurrentState.Data));
+				UProperty* Inner = ArrayProperty->Inner;
+
+				// push elements on stack (in reverse order)
+				for (int Index = ArrayHelper.Num() - 1; Index >= 0; --Index)
+				{
+					FWriteState NewState;
+
+					NewState.Data = ArrayHelper.GetRawPtr(Index);
+					NewState.Property = Inner;
+					NewState.TypeInfo = Inner->GetClass();
+					NewState.HasBeenProcessed = false;
+
+					StateStack.Push(NewState);
+				}
+			}
+			else
+			{
+				Backend.EndArray(CurrentState.Property);
+			}
+		}
+
+		// static arrays
+		else if (CurrentState.Property->ArrayDim > 1)
+		{
+			Backend.BeginArray(CurrentState.Property);
+
+			for (int32 ArrayIndex = 0; ArrayIndex < CurrentState.Property->ArrayDim; ++ArrayIndex)
+			{
+				Backend.WriteProperty(CurrentState.Property, CurrentState.Data, CurrentState.TypeInfo, ArrayIndex);
+			}
+
+			Backend.EndArray(CurrentState.Property);
+		}
+
+		// all other properties
+		else
+		{
+			Backend.WriteProperty(CurrentState.Property, CurrentState.Data, CurrentState.TypeInfo);
+		}
+	}
+}
diff --git a/Engine/Source/Runtime/Serialization/Private/Tests/StructSerializerTest.cpp b/Engine/Source/Runtime/Serialization/Private/Tests/StructSerializerTest.cpp
new file mode 100644
--- /dev/null
+++ b/Engine/Source/Runtime/Serialization/Private/Tests/StructSerializerTest.cpp
@@ -0,0 +1,91 @@
+// Copyright 1998-2013 Epic Games, Inc. All Rights Reserved.
+
+#include "SerializationPrivatePCH.h"
+#include "AutomationTest.h"
+#include "JsonStructSerializerBackend.h"
+#include "StructDeserializer.h"
+#include "StructSerializer.h"
+#include "StructSerializerTestTypes.h"
+
+
+/* Internal helpers
+ *****************************************************************************/
+
+namespace StructSerializerTest
+{
+	void TestSerialization( FAutomationTestBase& Test, IStructSerializerBackend& SerializerBackend, IStructDeserializerBackend& DeserializerBackend )
+	{
+		// serialization
+		FStructSerializerTestStruct TestStruct;
+		{
+			FStructSerializer::Serialize(TestStruct, SerializerBackend);
+		}
+
+		// deserialization
+		FStructSerializerTestStruct TestStruct2(NoInit);
+		{
+			FStructDeserializerPolicies Policies;
+			Policies.MissingFields = EStructDeserializerErrorPolicies::Warning;
+			
+			Test.TestTrue(TEXT("Deserialization must succeed"), FStructDeserializer::Deserialize(TestStruct2, DeserializerBackend, Policies));
+		}
+
+		// test for equality
+		Test.TestEqual<int8>(TEXT("Numerics.Int8 value must be the same before and after de-/serialization"), TestStruct.Numerics.Int8, TestStruct2.Numerics.Int8);
+		Test.TestEqual<int16>(TEXT("Numerics.Int16 value must be the same before and after de-/serialization"), TestStruct.Numerics.Int16, TestStruct2.Numerics.Int16);
+		Test.TestEqual<int32>(TEXT("Numerics.Int32 value must be the same before and after de-/serialization"), TestStruct.Numerics.Int32, TestStruct2.Numerics.Int32);
+		Test.TestEqual<int64>(TEXT("Numerics.Int64 value must be the same before and after de-/serialization"), TestStruct.Numerics.Int64, TestStruct2.Numerics.Int64);
+		Test.TestEqual<uint8>(TEXT("Numerics.UInt8 value must be the same before and after de-/serialization"), TestStruct.Numerics.UInt8, TestStruct2.Numerics.UInt8);
+		Test.TestEqual<uint16>(TEXT("Numerics.UInt16 value must be the same before and after de-/serialization"), TestStruct.Numerics.UInt16, TestStruct2.Numerics.UInt16);
+		Test.TestEqual<uint32>(TEXT("Numerics.UInt32 value must be the same before and after de-/serialization"), TestStruct.Numerics.UInt32, TestStruct2.Numerics.UInt32);
+		Test.TestEqual<uint64>(TEXT("Numerics.UInt64 value must be the same before and after de-/serialization"), TestStruct.Numerics.UInt64, TestStruct2.Numerics.UInt64);
+		Test.TestEqual<float>(TEXT("Numerics.Float value must be the same before and after de-/serialization"), TestStruct.Numerics.Float, TestStruct2.Numerics.Float);
+		Test.TestEqual<double>(TEXT("Numerics.Double value must be the same before and after de-/serialization"), TestStruct.Numerics.Double, TestStruct2.Numerics.Double);
+
+		Test.TestEqual<bool>(TEXT("Booleans.BoolFalse must be the same before and after de-/serialization"), TestStruct.Booleans.BoolFalse, TestStruct2.Booleans.BoolFalse);
+		Test.TestEqual<bool>(TEXT("Booleans.BoolTrue must be the same before and after de-/serialization"), TestStruct.Booleans.BoolTrue, TestStruct2.Booleans.BoolTrue);
+		Test.TestEqual<uint32>(TEXT("Booleans.Bitfield must be the same before and after de-/serialization"), TestStruct.Booleans.Bitfield, TestStruct2.Booleans.Bitfield);
+
+		Test.TestEqual<TSubclassOf<class UObject>>(TEXT("Objects.Class must be the same before and after de-/serialization"), TestStruct.Objects.Class, TestStruct2.Objects.Class);
+		Test.TestEqual<UObject*>(TEXT("Objects.ObjectPtr must be the same before and after de-/serialization"), TestStruct.Objects.ObjectPtr, TestStruct2.Objects.ObjectPtr);
+
+		Test.TestEqual<FName>(TEXT("Builtins.Name must be the same before and after de-/serialization"), TestStruct.Builtins.Name, TestStruct2.Builtins.Name);
+		Test.TestEqual<FString>(TEXT("Builtins.String must be the same before and after de-/serialization"), TestStruct.Builtins.String, TestStruct2.Builtins.String);
+		Test.TestEqual<FRotator>(TEXT("Builtins.Rotator must be the same before and after de-/serialization"), TestStruct.Builtins.Rotator, TestStruct2.Builtins.Rotator);
+		Test.TestEqual<FVector>(TEXT("Builtins.Vector must be the same before and after de-/serialization"), TestStruct.Builtins.Vector, TestStruct2.Builtins.Vector);
+
+		Test.TestEqual<TArray<int32>>(TEXT("Arrays.Int32Array must be the same before and after de-/serialization"), TestStruct.Arrays.Int32Array, TestStruct2.Arrays.Int32Array);
+		Test.TestEqual<int32>(TEXT("Arrays.StaticSingleElement[0] must be the same before and after de-/serialization"), TestStruct.Arrays.StaticSingleElement[0], TestStruct2.Arrays.StaticSingleElement[0]);
+		Test.TestEqual<int32>(TEXT("Arrays.StaticInt32Array[0] must be the same before and after de-/serialization"), TestStruct.Arrays.StaticInt32Array[0], TestStruct2.Arrays.StaticInt32Array[0]);
+		Test.TestEqual<int32>(TEXT("Arrays.StaticInt32Array[1] must be the same before and after de-/serialization"), TestStruct.Arrays.StaticInt32Array[1], TestStruct2.Arrays.StaticInt32Array[1]);
+		Test.TestEqual<int32>(TEXT("Arrays.StaticInt32Array[2] must be the same before and after de-/serialization"), TestStruct.Arrays.StaticInt32Array[2], TestStruct2.Arrays.StaticInt32Array[2]);
+		Test.TestEqual<float>(TEXT("Arrays.StaticFloatArray[0] must be the same before and after de-/serialization"), TestStruct.Arrays.StaticFloatArray[0], TestStruct2.Arrays.StaticFloatArray[0]);
+		Test.TestEqual<float>(TEXT("Arrays.StaticFloatArray[1] must be the same before and after de-/serialization"), TestStruct.Arrays.StaticFloatArray[1], TestStruct2.Arrays.StaticFloatArray[1]);
+		Test.TestEqual<float>(TEXT("Arrays.StaticFloatArray[2] must be the same before and after de-/serialization"), TestStruct.Arrays.StaticFloatArray[2], TestStruct2.Arrays.StaticFloatArray[2]);
+		Test.TestEqual<TArray<FVector>>(TEXT("Arrays.VectorArray must be the same before and after de-/serialization"), TestStruct.Arrays.VectorArray, TestStruct2.Arrays.VectorArray);
+	}
+}
+
+
+/* Tests
+ *****************************************************************************/
+
+IMPLEMENT_SIMPLE_AUTOMATION_TEST(FJsonStructSerializerTest, "Core.Serialization.JsonStructSerializer", EAutomationTestFlags::ATF_Editor)
+
+
+bool FJsonStructSerializerTest::RunTest( const FString& Parameters )
+{
+	// json
+	{
+		TArray<uint8> Buffer;
+		FMemoryReader Reader(Buffer);
+		FMemoryWriter Writer(Buffer);
+
+		FJsonStructSerializerBackend SerializerBackend(Writer);
+		FJsonStructDeserializerBackend DeserializerBackend(Reader);
+
+		StructSerializerTest::TestSerialization(*this, SerializerBackend, DeserializerBackend);
+	}
+
+	return true;
+}
diff --git a/Engine/Source/Runtime/Serialization/Private/Tests/StructSerializerTestTypes.h b/Engine/Source/Runtime/Serialization/Private/Tests/StructSerializerTestTypes.h
new file mode 100644
--- /dev/null
+++ b/Engine/Source/Runtime/Serialization/Private/Tests/StructSerializerTestTypes.h
@@ -0,0 +1,231 @@
+// Copyright 1998-2013 Epic Games, Inc. All Rights Reserved.
+
+#include "StructSerializerTestTypes.generated.h"
+
+
+/**
+ * Test structure for numeric properties.
+ */
+USTRUCT()
+struct FStructSerializerNumericTestStruct
+{
+	GENERATED_USTRUCT_BODY()
+
+	UPROPERTY()
+	int8 Int8;
+
+	UPROPERTY()
+	int16 Int16;
+
+	UPROPERTY()
+	int32 Int32;
+
+	UPROPERTY()
+	int64 Int64;
+
+	UPROPERTY()
+	uint8 UInt8;
+
+	UPROPERTY()
+	uint16 UInt16;
+
+	UPROPERTY()
+	uint32 UInt32;
+
+	UPROPERTY()
+	uint64 UInt64;
+
+	UPROPERTY()
+	float Float;
+
+	UPROPERTY()
+	double Double;
+
+	/** Default constructor. */
+	FStructSerializerNumericTestStruct( )
+		: Int8(-8)
+		, Int16(-16)
+		, Int32(-32)
+		, Int64(-64)
+		, UInt8(8)
+		, UInt16(16)
+		, UInt32(32)
+		, UInt64(64)
+		, Float(4.125)
+		, Double(1.03125)
+	{ }
+
+	/** Creates an uninitialized instance. */
+	FStructSerializerNumericTestStruct( ENoInit ) { }
+};
+
+
+/**
+ * Test structure for boolean properties.
+ */
+USTRUCT()
+struct FStructSerializerBooleanTestStruct
+{
+	GENERATED_USTRUCT_BODY()
+
+	UPROPERTY()
+	bool BoolFalse;
+
+	UPROPERTY()
+	bool BoolTrue;
+
+	UPROPERTY()
+	uint32 Bitfield;
+
+	/** Default constructor. */
+	FStructSerializerBooleanTestStruct( )
+		: BoolFalse(false)
+		, BoolTrue(true)
+		, Bitfield(1)
+	{ }
+
+	/** Creates an uninitialized instance. */
+	FStructSerializerBooleanTestStruct( ENoInit ) { }
+};
+
+
+/**
+ * Test structure for UObject properties.
+ */
+USTRUCT()
+struct FStructSerializerObjectTestStruct
+{
+	GENERATED_USTRUCT_BODY()
+
+	UPROPERTY()
+	TSubclassOf<class UObject> Class;
+
+	UPROPERTY()
+	class UObject* ObjectPtr;
+
+	/** Default constructor. */
+	FStructSerializerObjectTestStruct( )
+		: ObjectPtr(nullptr)
+	{ }
+
+	/** Creates an uninitialized instance. */
+	FStructSerializerObjectTestStruct( ENoInit ) { }
+};
+
+
+/**
+ * Test structure for properties of various built-in types.
+ */
+USTRUCT()
+struct FStructSerializerBuiltinTestStruct
+{
+	GENERATED_USTRUCT_BODY()
+
+	UPROPERTY()
+	FName Name;
+
+	UPROPERTY()
+	FString String;
+
+	UPROPERTY()
+	FRotator Rotator;
+
+	UPROPERTY()
+	FVector Vector;
+
+	/** Default constructor. */
+	FStructSerializerBuiltinTestStruct( )
+		: String("Hello World!")
+		, Rotator(4096, 8192, 16384)
+		, Vector(1.0f, 2.0f, 3.0f)
+	{ }
+
+	/** Creates an uninitialized instance. */
+	FStructSerializerBuiltinTestStruct( ENoInit ) { }
+};
+
+
+/**
+ * Test structure for array properties.
+ */
+USTRUCT()
+struct FStructSerializerArrayTestStruct
+{
+	GENERATED_USTRUCT_BODY()
+
+	UPROPERTY()
+	TArray<int32> Int32Array;
+
+	UPROPERTY()
+	int32 StaticSingleElement[1];
+
+	UPROPERTY()
+	int32 StaticInt32Array[3];
+
+	UPROPERTY()
+	float StaticFloatArray[3];
+
+	UPROPERTY()
+	TArray<FVector> VectorArray;
+
+	/** Default constructor. */
+	FStructSerializerArrayTestStruct( )
+	{
+		Int32Array.Add(-1);
+		Int32Array.Add(0);
+		Int32Array.Add(1);
+
+		StaticSingleElement[0] = 42;
+
+		StaticInt32Array[0] = -1;
+		StaticInt32Array[1] = 0;
+		StaticInt32Array[2] = 1;
+
+		StaticFloatArray[0] = -1.0f;
+		StaticFloatArray[1] = 0.0f;
+		StaticFloatArray[2] = 1.0f;
+
+		VectorArray.Add(FVector(1.0f, 2.0f, 3.0f));
+		VectorArray.Add(FVector(-1.0f, -2.0f, -3.0f));
+	}
+
+	/** Creates an uninitialized instance. */
+	FStructSerializerArrayTestStruct( ENoInit ) { }
+};
+
+
+/**
+ * Test structure for all supported types.
+ */
+USTRUCT()
+struct FStructSerializerTestStruct
+{
+	GENERATED_USTRUCT_BODY()
+
+	UPROPERTY()
+	FStructSerializerNumericTestStruct Numerics;
+
+	UPROPERTY()
+	FStructSerializerBooleanTestStruct Booleans;
+
+	UPROPERTY()
+	FStructSerializerObjectTestStruct Objects;
+
+	UPROPERTY()
+	FStructSerializerBuiltinTestStruct Builtins;
+
+	UPROPERTY()
+	FStructSerializerArrayTestStruct Arrays;
+
+	/** Default constructor. */
+	FStructSerializerTestStruct() { }
+
+	/** Creates an uninitialized instance. */
+	FStructSerializerTestStruct( ENoInit )
+		: Numerics(NoInit)
+		, Booleans(NoInit)
+		, Objects(NoInit)
+		, Builtins(NoInit)
+		, Arrays(NoInit)
+	{ }
+};
diff --git a/Engine/Source/Runtime/Serialization/Public/Backends/JsonStructDeserializerBackend.h b/Engine/Source/Runtime/Serialization/Public/Backends/JsonStructDeserializerBackend.h
new file mode 100644
--- /dev/null
+++ b/Engine/Source/Runtime/Serialization/Public/Backends/JsonStructDeserializerBackend.h
@@ -0,0 +1,52 @@
+// Copyright 1998-2013 Epic Games, Inc. All Rights Reserved.
+
+#pragma once
+
+#include "IStructDeserializerBackend.h"
+
+
+/**
+ * Implements a reader for UStruct deserialization using Json.
+ *
+ * Note: The underlying Json de-serializer is currently hard-coded to use UCS2CHAR.
+ * This is because the current JsonReader API does not allow writers to be substituted since it's
+ * all based on templates. At some point we will refactor the low-level Json API to provide more
+ * flexibility for serialization.
+ */
+class SERIALIZATION_API FJsonStructDeserializerBackend
+	: public IStructDeserializerBackend
+{
+public:
+
+	/**
+	 * Creates and initializes a new instance.
+	 *
+	 * @param Archive The archive to deserialize from.
+	 */
+	FJsonStructDeserializerBackend( FArchive& Archive )
+		: JsonReader(TJsonReader<UCS2CHAR>::Create(&Archive))
+	{ }
+
+public:
+
+	// IStructDeserializerBackend interface
+
+	virtual const FString& GetCurrentPropertyName( ) const override;
+	virtual FString GetDebugString( ) const override;
+	virtual const FString& GetLastErrorMessage( ) const override;
+	virtual bool GetNextToken( EStructDeserializerBackendTokens& OutToken ) override;
+	virtual bool ReadProperty( UProperty* Property, UProperty* Outer, void* Data, int32 ArrayIndex ) override;
+	virtual void SkipArray( ) override;
+	virtual void SkipStructure( ) override;
+
+private:
+
+	/** Holds the name of the last read Json identifier. */
+	FString LastIdentifier;
+
+	/** Holds the last read Json notation. */
+	EJsonNotation LastNotation;
+
+	/** Holds the Json reader used for the actual reading of the archive. */
+	TSharedRef<TJsonReader<>> JsonReader;
+};
diff --git a/Engine/Source/Runtime/Serialization/Public/Backends/JsonStructSerializerBackend.h b/Engine/Source/Runtime/Serialization/Public/Backends/JsonStructSerializerBackend.h
new file mode 100644
--- /dev/null
+++ b/Engine/Source/Runtime/Serialization/Public/Backends/JsonStructSerializerBackend.h
@@ -0,0 +1,46 @@
+// Copyright 1998-2013 Epic Games, Inc. All Rights Reserved.
+
+#pragma once
+
+#include "IStructSerializerBackend.h"
+
+
+/**
+ * Implements a writer for UStruct serialization using Json.
+ *
+ * Note: The underlying Json serializer is currently hard-coded to use UCS2CHAR and pretty-print.
+ * This is because the current JsonWriter API does not allow writers to be substituted since it's
+ * all based on templates. At some point we will refactor the low-level Json API to provide more
+ * flexibility for serialization.
+ */
+class SERIALIZATION_API FJsonStructSerializerBackend
+	: public IStructSerializerBackend
+{
+public:
+
+	/**
+	 * Creates and initializes a new instance.
+	 *
+	 * @param Archive The archive to serialize into.
+	 */
+	FJsonStructSerializerBackend( FArchive& Archive )
+		: JsonWriter(TJsonWriter<UCS2CHAR>::Create(&Archive))
+	{ }
+
+public:
+
+	// IStructSerializerBackend interface
+
+	virtual void BeginArray( UProperty* Property ) override;
+	virtual void BeginStructure( UProperty* Property ) override;
+	virtual void BeginStructure( UStruct* TypeInfo ) override;
+	virtual void EndArray( UProperty* Property ) override;
+	virtual void EndStructure( ) override;
+	virtual void WriteComment( const FString& Comment ) override;
+	virtual void WriteProperty( UProperty* Property, const void* Data, UStruct* TypeInfo, int32 ArrayIndex ) override;
+
+private:
+
+	/** Holds the Json writer used for the actual serialization. */
+	TSharedRef<TJsonWriter<>> JsonWriter;
+};
diff --git a/Engine/Source/Runtime/Serialization/Public/Interfaces/IStructDeserializerBackend.h b/Engine/Source/Runtime/Serialization/Public/Interfaces/IStructDeserializerBackend.h
new file mode 100644
--- /dev/null
+++ b/Engine/Source/Runtime/Serialization/Public/Interfaces/IStructDeserializerBackend.h
@@ -0,0 +1,102 @@
+// Copyright 1998-2013 Epic Games, Inc. All Rights Reserved.
+
+#pragma once
+
+
+class UProperty;
+class UStruct;
+
+
+/**
+ * Enumerates serialization tokens.
+ */
+enum class EStructDeserializerBackendTokens
+{
+	/** End of an array. */
+	ArrayEnd,
+
+	/** Beginning of an array. */
+	ArrayStart,
+
+	/** A comment. */
+	Comment,
+
+	/** An error occurred when reading the token. */
+	Error,
+
+	/** No token available. */
+	None,
+
+	/** A scalar property. */
+	Property,
+
+	/** End of a data structure. */
+	StructureEnd,
+
+	/** Beginning of a data structure. */
+	StructureStart,
+};
+
+
+/**
+ * Interface for UStruct serializer backends.
+ */
+class IStructDeserializerBackend
+{
+public:
+
+	/**
+	 * Gets the identifier of the current field.
+	 *
+	 * @return The field's name.
+	 */
+	virtual const FString& GetCurrentPropertyName( ) const = 0;
+
+	/**
+	 * Gets a debug string for the reader's current state.
+	 *
+	 * The returned string contains debug information that is relevant to the reader's
+	 * serialization format. For example, it could be a line and column number for text
+	 * based formats, or a byte offset for binary serialization formats.
+	 *
+	 * @return The debug string.
+	 */
+	virtual FString GetDebugString( ) const = 0;
+
+	/**
+	 * Gets the last error message.
+	 *
+	 * @return Error message.
+	 */
+	virtual const FString& GetLastErrorMessage( ) const = 0;
+
+	/**
+	 * Reads the next token from the stream.
+	 *
+	 * @param OutToken Will contain the read token, if any.
+	 * @return true if a token was read, false otherwise.
+	 */
+	virtual bool GetNextToken( EStructDeserializerBackendTokens& OutToken ) = 0;
+
+	/**
+	 * Reads the specified property from the stream.
+	 *
+	 * @param Property The property to read into.
+	 * @param Outer The outer property holding the property to read (in case of arrays).
+	 * @param Data The buffer that will hold the read data.
+	 * @param ArrayIndex - An index into the property array (for static arrays).
+	 * @return true on success, false otherwise.
+	 */
+	virtual bool ReadProperty( UProperty* Property, UProperty* Outer, void* Data, int32 ArrayIndex ) = 0;
+
+	/** Skips the array that is currently being read from the stream. */
+	virtual void SkipArray( ) = 0;
+
+	/** Skips the object that is currently being read from the stream. */
+	virtual void SkipStructure( ) = 0;
+
+public:
+
+	/** Virtual destructor. */
+	virtual ~IStructDeserializerBackend( ) { }
+};
diff --git a/Engine/Source/Runtime/Serialization/Public/Interfaces/IStructSerializerBackend.h b/Engine/Source/Runtime/Serialization/Public/Interfaces/IStructSerializerBackend.h
new file mode 100644
--- /dev/null
+++ b/Engine/Source/Runtime/Serialization/Public/Interfaces/IStructSerializerBackend.h
@@ -0,0 +1,73 @@
+// Copyright 1998-2013 Epic Games, Inc. All Rights Reserved.
+
+#pragma once
+
+
+class UProperty;
+class UStruct;
+
+
+/**
+ * Interface for UStruct serializer backends.
+ */
+class IStructSerializerBackend
+{
+public:
+
+	/**
+	 * Signals the beginning of an array.
+	 *
+	 * @param Property The property that holds the array.
+	 */
+	virtual void BeginArray( UProperty* Property ) = 0;
+
+	/**
+	 * Signals the beginning of a child structure.
+	 *
+	 * @param Property The property that holds the object.
+	 */
+	virtual void BeginStructure( UProperty* Property ) = 0;
+
+	/**
+	 * Signals the beginning of a root structure.
+	 *
+	 * @param TypeInfo The object's type information.
+	 */
+	virtual void BeginStructure( UStruct* TypeInfo ) = 0;
+
+	/**
+	 * Signals the end of an array.
+	 *
+	 * @param Property The property that holds the array.
+	 */
+	virtual void EndArray( UProperty* Property ) = 0;
+
+	/**
+	 * Signals the end of an object.
+	 */
+	virtual void EndStructure( ) = 0;
+
+	/**
+	 * Writes a comment to the output stream.
+	 *
+	 * @param Comment The comment text.
+	 */
+	virtual void WriteComment( const FString& Comment ) = 0;
+
+	/**
+	 * Writes a property to the output stream.
+	 *
+	 * Depending on the context, properties to be written can be either object properties or array elements.
+	 *
+	 * @param Property The property that holds the data to write.
+	 * @param Data The property's data to write.
+	 * @param TypeInfo The property data's type information.
+	 * @param ArrayIndex An index into the property array (for static arrays).
+	 */
+	virtual void WriteProperty( UProperty* Property, const void* Data, UStruct* TypeInfo, int32 ArrayIndex = 0 ) = 0;
+
+public:
+
+	/** Virtual destructor. */
+	virtual ~IStructSerializerBackend( ) { }
+};
diff --git a/Engine/Source/Runtime/Serialization/Public/StructDeserializer.h b/Engine/Source/Runtime/Serialization/Public/StructDeserializer.h
new file mode 100644
--- /dev/null
+++ b/Engine/Source/Runtime/Serialization/Public/StructDeserializer.h
@@ -0,0 +1,105 @@
+// Copyright 1998-2013 Epic Games, Inc. All Rights Reserved.
+
+#pragma once
+
+
+class IStructDeserializerBackend;
+class UStruct;
+
+
+/**
+ * Enumerates policies for various errors during de-serialization.
+ */
+enum class EStructDeserializerErrorPolicies
+{
+	/** Fail the de-serialization. */
+	Error,
+		
+	/** Ignore the error (default). */
+	Ignore,
+		
+	/** Print a warning to the log. */
+	Warning,
+};
+
+
+/**
+ * Structure for UStruct serialization policies.
+ */
+struct FStructDeserializerPolicies
+{
+	/** Holds the policy for handling missing fields. */
+	EStructDeserializerErrorPolicies MissingFields;
+
+	/** Default constructor. */
+	FStructDeserializerPolicies( )
+		: MissingFields(EStructDeserializerErrorPolicies::Ignore)
+	{ }
+};
+
+
+/**
+ * Implements a static class that can deserialize UStruct based types.
+ *
+ * This class implements the basic functionality for the serialization of UStructs, such as
+ * iterating a structure's properties and writing property values. The actual reading of
+ * serialized input data is performed by de-serialization backends, which allows this
+ * class to remain serialization format agnostic.
+ */
+class FStructDeserializer
+{
+public:
+
+	/**
+	 * Deserializes a data structure from an archive using the specified policy.
+	 *
+	 * @param OutStruct A pointer to the data structure to deserialize into.
+	 * @param TypeInfo The data structure's type information.
+	 * @param Backend The de-serialization backend to use.
+	 * @param Policies The de-serialization policies to use.
+	 * @return true if deserialization was successful, false otherwise.
+	 */
+	SERIALIZATION_API static bool Deserialize( void* OutStruct, UStruct& TypeInfo, IStructDeserializerBackend& Backend, const FStructDeserializerPolicies& Policies );
+
+	/**
+	 * Deserializes a data structure from an archive using the default policy.
+	 *
+	 * @param OutStruct A pointer to the data structure to deserialize into.
+	 * @param TypeInfo The data structure's type information.
+	 * @param Backend The de-serialization backend to use.
+	 * @return true if deserialization was successful, false otherwise.
+	 */
+	static bool Deserialize( void* OutStruct, UStruct& TypeInfo, IStructDeserializerBackend& Backend )
+	{
+		return Deserialize(OutStruct, TypeInfo, Backend, FStructDeserializerPolicies());
+	}
+
+public:
+
+	/**
+	 * Deserializes a data structure from an archive using the default policy.
+	 *
+	 * @param OutStruct The struct to deserialize into.
+	 * @param Backend The de-serialization backend to use.
+	 * @return true if deserialization was successful, false otherwise.
+	 */
+	template<typename StructType>
+	static bool Deserialize( StructType& OutStruct, IStructDeserializerBackend& Backend )
+	{
+		return Deserialize(&OutStruct, *StructType::StaticStruct(), Backend);
+	}
+
+	/**
+	 * Deserializes a data structure from an archive using the specified policy.
+	 *
+	 * @param OutStruct The struct to deserialize into.
+	 * @param Backend The de-serialization backend to use.
+	 * @param Policies The de-serialization policies to use.
+	 * @return true if deserialization was successful, false otherwise.
+	 */
+	template<typename StructType>
+	static bool Deserialize( StructType& OutStruct, IStructDeserializerBackend& Backend, const FStructDeserializerPolicies& Policies )
+	{
+		return Deserialize(&OutStruct, *StructType::StaticStruct(), Backend, Policies);
+	}
+};
diff --git a/Engine/Source/Runtime/Serialization/Public/StructSerializer.h b/Engine/Source/Runtime/Serialization/Public/StructSerializer.h
new file mode 100644
--- /dev/null
+++ b/Engine/Source/Runtime/Serialization/Public/StructSerializer.h
@@ -0,0 +1,127 @@
+// Copyright 1998-2013 Epic Games, Inc. All Rights Reserved.
+
+#pragma once
+
+
+class IStructSerializerBackend;
+class UStruct;
+
+
+/**
+ * Enumerates policies for serializing null values.
+ */
+enum class EStructSerializerNullValuePolicies
+{
+	/** Do not serialize NULL values. */
+	Ignore,
+		
+	/** Serialize NULL values. */
+	Serialize
+};
+
+
+/**
+ * Enumerates policies for serializing object reference loops.
+ */
+enum class EStructSerializerReferenceLoopPolicies
+{
+	/** Circular references generate an error. */
+	Error,
+		
+	/** Ignore circular references. */
+	Ignore,
+		
+	/** Serialize circular references. */
+	Serialize
+};
+
+
+/**
+ * Structure for UStruct serialization policies.
+ */
+struct FStructSerializerPolicies
+{
+	/** Holds the policy for NULL values. */
+	EStructSerializerNullValuePolicies NullValues;
+
+	/** Holds the policy for reference loops. */
+	EStructSerializerReferenceLoopPolicies ReferenceLoops;
+
+	/** Default constructor. */
+	FStructSerializerPolicies( )
+		: NullValues(EStructSerializerNullValuePolicies::Serialize)
+		, ReferenceLoops(EStructSerializerReferenceLoopPolicies::Ignore)
+	{ }
+};
+
+
+/**
+ * Implements a static class that can serialize UStruct based types.
+ *
+ * This class implements the basic functionality for the serialization of UStructs, such as
+ * iterating a structure's properties and writing property values. The actual writing of
+ * serialized output data is performed by serialization backends, which allows this class
+ * to remain serialization format agnostic.
+ *
+ * The serializer's behavior can be customized with serialization policies. This allows for
+ * control over how to handle null values, circular references and other edge cases.
+ */
+class FStructSerializer
+{
+public:
+
+	/**
+	 * Serializes a given data structure of the specified type using the specified policy.
+	 *
+	 * @param Struct The data structure to serialize.
+	 * @param TypeInfo The structure's type information.
+	 * @param Backend The serialization backend to use.
+	 * @param Policies The serialization policies to use.
+	 * @see Deserialize
+	 */
+	SERIALIZATION_API static void Serialize( const void* Struct, UStruct& TypeInfo, IStructSerializerBackend& Backend, const FStructSerializerPolicies& Policies );
+
+	/**
+	 * Serializes a given data structure of the specified type using the default policy.
+	 *
+	 * @param Struct The data structure to serialize.
+	 * @param TypeInfo The structure's type information.
+	 * @param Backend The serialization backend to use.
+	 * @see Deserialize
+	 */
+	static void Serialize( const void* Struct, UStruct& TypeInfo, IStructSerializerBackend& Backend )
+	{
+		Serialize(Struct, TypeInfo, Backend, FStructSerializerPolicies());
+	}
+
+public:
+
+	/**
+	 * Serializes a given USTRUCT to a string using the default policy.
+	 *
+	 * @param StructType The type of the struct to serialize.
+	 * @param Struct The struct to serialize.
+	 * @param Backend The serialization backend to use.
+	 * @return A string holding the serialized object.
+	 */
+	template<typename StructType>
+	static void Serialize( const StructType& Struct, IStructSerializerBackend& Backend )
+	{
+		Serialize(&Struct, *Struct.StaticStruct(), Backend, FStructSerializerPolicies());
+	}
+
+	/**
+	 * Serializes a given USTRUCT to a string using the specified policy.
+	 *
+	 * @param StructType The type of the struct to serialize.
+	 * @param Struct The struct to serialize.
+	 * @param Backend The serialization backend to use.
+	 * @param Policies The serialization policies to use.
+	 * @return A string holding the serialized object.
+	 */
+	template<typename StructType>
+	static void Serialize( const StructType& Struct, IStructSerializerBackend& Backend, const FStructSerializerPolicies& Policies )
+	{
+		Serialize(&Struct, *Struct.StaticStruct(), Backend, Policies);
+	}
+};
diff --git a/Engine/Source/Runtime/Serialization/Serialization.Build.cs b/Engine/Source/Runtime/Serialization/Serialization.Build.cs
new file mode 100644
--- /dev/null
+++ b/Engine/Source/Runtime/Serialization/Serialization.Build.cs
@@ -0,0 +1,25 @@
+// Copyright 1998-2013 Epic Games, Inc. All Rights Reserved.
+
+namespace UnrealBuildTool.Rules
+{
+	public class Serialization : ModuleRules
+	{
+		public Serialization(TargetInfo Target)
+		{	
+			PrivateDependencyModuleNames.AddRange(
+				new string[]
+				{
+					"Core",
+					"CoreUObject",
+					"Json",
+				}
+			);
+
+			PrivateIncludePaths.AddRange(
+				new string[] {
+					"Runtime/Serialization/Private",
+				}
+			);
+		}
+	}
+}
