Serialization: Fixed template parameter

diff --git a/Engine/Source/Runtime/Serialization/Public/Backends/JsonStructDeserializerBackend.h b/Engine/Source/Runtime/Serialization/Public/Backends/JsonStructDeserializerBackend.h
--- a/Engine/Source/Runtime/Serialization/Public/Backends/JsonStructDeserializerBackend.h
+++ b/Engine/Source/Runtime/Serialization/Public/Backends/JsonStructDeserializerBackend.h
@@ -48,5 +48,5 @@ private:
 	EJsonNotation LastNotation;
 
 	/** Holds the Json reader used for the actual reading of the archive. */
-	TSharedRef<TJsonReader<>> JsonReader;
+	TSharedRef<TJsonReader<UCS2CHAR>> JsonReader;
 };
diff --git a/Engine/Source/Runtime/Serialization/Public/Backends/JsonStructSerializerBackend.h b/Engine/Source/Runtime/Serialization/Public/Backends/JsonStructSerializerBackend.h
--- a/Engine/Source/Runtime/Serialization/Public/Backends/JsonStructSerializerBackend.h
+++ b/Engine/Source/Runtime/Serialization/Public/Backends/JsonStructSerializerBackend.h
@@ -42,5 +42,5 @@ public:
 private:
 
 	/** Holds the Json writer used for the actual serialization. */
-	TSharedRef<TJsonWriter<>> JsonWriter;
+	TSharedRef<TJsonWriter<UCS2CHAR>> JsonWriter;
 };
