Serialization: Fixed template parameters

diff --git a/Engine/Source/Runtime/Serialization/Private/Backends/JsonStructSerializerBackend.cpp b/Engine/Source/Runtime/Serialization/Private/Backends/JsonStructSerializerBackend.cpp
--- a/Engine/Source/Runtime/Serialization/Private/Backends/JsonStructSerializerBackend.cpp
+++ b/Engine/Source/Runtime/Serialization/Private/Backends/JsonStructSerializerBackend.cpp
@@ -12,7 +12,7 @@ namespace JsonStructSerializerBackend
 {
 	// Writes a property value to the serialization output.
 	template<typename ValueType>
-	void WritePropertyValue( const TSharedRef<TJsonWriter<TCHAR, TPrettyJsonPrintPolicy<TCHAR>>> JsonWriter, UProperty* Property, const ValueType& Value )
+	void WritePropertyValue( const TSharedRef<TJsonWriter<UCS2CHAR>> JsonWriter, UProperty* Property, const ValueType& Value )
 	{
 		if ((Property == nullptr) || (Property->ArrayDim > 1) || (Property->GetOuter()->GetClass() == UArrayProperty::StaticClass()))
 		{
@@ -25,7 +25,7 @@ namespace JsonStructSerializerBackend
 	}
 
 	// Writes a null value to the serialization output.
-	void WriteNull( const TSharedRef<TJsonWriter<TCHAR, TPrettyJsonPrintPolicy<TCHAR>>> JsonWriter, UProperty* Property )
+	void WriteNull( const TSharedRef<TJsonWriter<UCS2CHAR>> JsonWriter, UProperty* Property )
 	{
 		if ((Property == nullptr) || (Property->ArrayDim > 1) || (Property->GetOuter()->GetClass() == UArrayProperty::StaticClass()))
 		{
