non unity compile fixes (penance for breaking it yesterday)

diff --git a/Engine/Source/Editor/GraphEditor/Private/AnimationNodes/SGraphNodeLayeredBoneBlend.cpp b/Engine/Source/Editor/GraphEditor/Private/AnimationNodes/SGraphNodeLayeredBoneBlend.cpp
--- a/Engine/Source/Editor/GraphEditor/Private/AnimationNodes/SGraphNodeLayeredBoneBlend.cpp
+++ b/Engine/Source/Editor/GraphEditor/Private/AnimationNodes/SGraphNodeLayeredBoneBlend.cpp
@@ -3,6 +3,7 @@
 
 #include "GraphEditorCommon.h"
 #include "SGraphNodeLayeredBoneBlend.h"
+#include "AnimGraphNode_LayeredBoneBlend.h"
 
 /////////////////////////////////////////////////////
 // SGraphNodeLayeredBoneBlend
diff --git a/Engine/Source/Runtime/Engine/Private/DemoNetDriver.cpp b/Engine/Source/Runtime/Engine/Private/DemoNetDriver.cpp
--- a/Engine/Source/Runtime/Engine/Private/DemoNetDriver.cpp
+++ b/Engine/Source/Runtime/Engine/Private/DemoNetDriver.cpp
@@ -10,6 +10,7 @@
 #include "Engine/DemoPendingNetGame.h"
 #include "Engine/ActorChannel.h"
 #include "RepLayout.h"
+#include "GameFramework/SpectatorPawn.h"
 
 DEFINE_LOG_CATEGORY_STATIC( LogDemo, Log, All );
 
diff --git a/Engine/Source/Runtime/Serialization/Private/Backends/JsonStructSerializerBackend.cpp b/Engine/Source/Runtime/Serialization/Private/Backends/JsonStructSerializerBackend.cpp
--- a/Engine/Source/Runtime/Serialization/Private/Backends/JsonStructSerializerBackend.cpp
+++ b/Engine/Source/Runtime/Serialization/Private/Backends/JsonStructSerializerBackend.cpp
@@ -1,7 +1,6 @@
 // Copyright 1998-2013 Epic Games, Inc. All Rights Reserved.
 
 #include "SerializationPrivatePCH.h"
-#include "Json.h"
 #include "JsonStructSerializerBackend.h"
 
 
diff --git a/Engine/Source/Runtime/Serialization/Private/Tests/StructSerializerTest.cpp b/Engine/Source/Runtime/Serialization/Private/Tests/StructSerializerTest.cpp
--- a/Engine/Source/Runtime/Serialization/Private/Tests/StructSerializerTest.cpp
+++ b/Engine/Source/Runtime/Serialization/Private/Tests/StructSerializerTest.cpp
@@ -3,6 +3,7 @@
 #include "SerializationPrivatePCH.h"
 #include "AutomationTest.h"
 #include "JsonStructSerializerBackend.h"
+#include "JsonStructDeserializerBackend.h"
 #include "StructDeserializer.h"
 #include "StructSerializer.h"
 #include "StructSerializerTestTypes.h"
diff --git a/Engine/Source/Runtime/Serialization/Public/Backends/JsonStructSerializerBackend.h b/Engine/Source/Runtime/Serialization/Public/Backends/JsonStructSerializerBackend.h
--- a/Engine/Source/Runtime/Serialization/Public/Backends/JsonStructSerializerBackend.h
+++ b/Engine/Source/Runtime/Serialization/Public/Backends/JsonStructSerializerBackend.h
@@ -3,7 +3,8 @@
 #pragma once
 
 #include "IStructSerializerBackend.h"
-
+#include "Json.h"
+#include "JsonWriter.h"
 
 /**
  * Implements a writer for UStruct serialization using Json.
