Serialization: Made the Json deserializer more forgiving if properties are not found or are of an unexpected type

diff --git a/Engine/Source/Runtime/Serialization/Private/Backends/JsonStructDeserializerBackend.cpp b/Engine/Source/Runtime/Serialization/Private/Backends/JsonStructDeserializerBackend.cpp
--- a/Engine/Source/Runtime/Serialization/Private/Backends/JsonStructDeserializerBackend.cpp
+++ b/Engine/Source/Runtime/Serialization/Private/Backends/JsonStructDeserializerBackend.cpp
@@ -9,32 +9,59 @@
 
 namespace JsonStructDeserializerBackend
 {
-	// Clears the value of the given property.
-	void ClearPropertyValue( UProperty* Property, UProperty* Outer, void* Data, int32 ArrayIndex )
+	/**
+	 * Clears the value of the given property.
+	 *
+	 * @param Property The property to clear.
+	 * @param Outer The property that contains the property to be cleared, if any.
+	 * @param Data A pointer to the memory holding the property's data.
+	 * @param ArrayIndex The index of the element to clear (if the property is an array).
+	 * @return true on success, false otherwise.
+	 * @see SetPropertyValue
+	 */
+	bool ClearPropertyValue( UProperty* Property, UProperty* Outer, void* Data, int32 ArrayIndex )
 	{
 		UArrayProperty* ArrayProperty = Cast<UArrayProperty>(Outer);
 
 		if (ArrayProperty != nullptr)
 		{
-			check(ArrayProperty->Inner == Property)
+			if (ArrayProperty->Inner != Property)
+			{
+				return false;
+			}
 
 			FScriptArrayHelper ArrayHelper(ArrayProperty, ArrayProperty->template ContainerPtrToValuePtr<void>(Data));
 			ArrayIndex = ArrayHelper.AddValue();
 		}
 
 		Property->ClearValue_InContainer(Data, ArrayIndex);
+
+		return true;
 	}
 
-	// Sets the value of the given property.
+
+	/**
+	 * Sets the value of the given property.
+	 *
+	 * @param Property The property to set.
+	 * @param Outer The property that contains the property to be set, if any.
+	 * @param Data A pointer to the memory holding the property's data.
+	 * @param ArrayIndex The index of the element to set (if the property is an array).
+	 * @return true on success, false otherwise.
+	 * @see ClearPropertyValue
+	 */
 	template<typename UPropertyType, typename PropertyType>
-	void SetPropertyValue( UProperty* Property, UProperty* Outer, void* Data, int32 ArrayIndex, const PropertyType& Value )
+	bool SetPropertyValue( UProperty* Property, UProperty* Outer, void* Data, int32 ArrayIndex, const PropertyType& Value )
 	{
 		PropertyType* ValuePtr = nullptr;
 		UArrayProperty* ArrayProperty = Cast<UArrayProperty>(Outer);
 
 		if (ArrayProperty != nullptr)
 		{
-			check(ArrayProperty->Inner == Property)
+			if (ArrayProperty->Inner != Property)
+			{
+				return false;
+			}
 
 			FScriptArrayHelper ArrayHelper(ArrayProperty, ArrayProperty->template ContainerPtrToValuePtr<void>(Data));
 			int32 Index = ArrayHelper.AddValue();
@@ -44,12 +71,23 @@ namespace JsonStructDeserializerBackend
 		else
 		{
 			UPropertyType* TypedProperty = Cast<UPropertyType>(Property);
-			check(TypedProperty)
+
+			if (TypedProperty == nullptr)
+			{
+				return false;
+			}
 
 			ValuePtr = TypedProperty->template ContainerPtrToValuePtr<PropertyType>(Data, ArrayIndex);
 		}
 
+		if (ValuePtr == nullptr)
+		{
+			return false;
+		}
+
 		*ValuePtr = Value;
+
+		return true;
 	}
 }
 
@@ -57,19 +95,19 @@ namespace JsonStructDeserializerBackend
 /* IStructDeserializerBackend interface
  *****************************************************************************/
 
-const FString& FJsonStructDeserializerBackend::GetCurrentPropertyName( ) const
+const FString& FJsonStructDeserializerBackend::GetCurrentPropertyName() const
 {
 	return JsonReader->GetIdentifier();
 }
 
 
-FString FJsonStructDeserializerBackend::GetDebugString( ) const
+FString FJsonStructDeserializerBackend::GetDebugString() const
 {
 	return FString::Printf(TEXT("Line: %u, Ch: %u"), JsonReader->GetLineNumber(), JsonReader->GetCharacterNumber());
 }
 
 
-const FString& FJsonStructDeserializerBackend::GetLastErrorMessage( ) const
+const FString& FJsonStructDeserializerBackend::GetLastErrorMessage() const
 {
 	return JsonReader->GetErrorMessage();
 }
@@ -134,14 +172,12 @@ bool FJsonStructDeserializerBackend::ReadProperty( UProperty* Property, UPropert
 
 			if (Property->GetClass() == UBoolProperty::StaticClass())
 			{
-				SetPropertyValue<UBoolProperty, bool>(Property, Outer, Data, ArrayIndex, BoolValue);
+				return SetPropertyValue<UBoolProperty, bool>(Property, Outer, Data, ArrayIndex, BoolValue);
 			}
-			else
-			{
-				UE_LOG(LogSerialization, Verbose, TEXT("Boolean field %s with value '%s' is not supported in UProperty type %s (%s)"), *Property->GetFName().ToString(), BoolValue ? *(GTrue.ToString()) : *(GFalse.ToString()), *Property->GetClass()->GetName(), *GetDebugString());
 
-				return false;
-			}
+			UE_LOG(LogSerialization, Verbose, TEXT("Boolean field %s with value '%s' is not supported in UProperty type %s (%s)"), *Property->GetFName().ToString(), BoolValue ? *(GTrue.ToString()) : *(GFalse.ToString()), *Property->GetClass()->GetName(), *GetDebugString());
+
+			return false;
 		}
 		break;
 
@@ -152,57 +188,63 @@ bool FJsonStructDeserializerBackend::ReadProperty( UProperty* Property, UPropert
 
 			if (Property->GetClass() == UByteProperty::StaticClass())
 			{
-				SetPropertyValue<UByteProperty, int8>(Property, Outer, Data, ArrayIndex, (int8)NumericValue);
+				return SetPropertyValue<UByteProperty, int8>(Property, Outer, Data, ArrayIndex, (int8)NumericValue);
 			}
-			else if (Property->GetClass() == UDoubleProperty::StaticClass())
+
+			if (Property->GetClass() == UDoubleProperty::StaticClass())
 			{
-				SetPropertyValue<UDoubleProperty, double>(Property, Outer, Data, ArrayIndex, NumericValue);
+				return SetPropertyValue<UDoubleProperty, double>(Property, Outer, Data, ArrayIndex, NumericValue);
 			}
-			else if (Property->GetClass() == UFloatProperty::StaticClass())
+			
+			if (Property->GetClass() == UFloatProperty::StaticClass())
 			{
-				SetPropertyValue<UFloatProperty, float>(Property, Outer, Data, ArrayIndex, NumericValue);
+				return SetPropertyValue<UFloatProperty, float>(Property, Outer, Data, ArrayIndex, NumericValue);
 			}
-			else if (Property->GetClass() == UIntProperty::StaticClass())
+			
+			if (Property->GetClass() == UIntProperty::StaticClass())
 			{
-				SetPropertyValue<UIntProperty, int32>(Property, Outer, Data, ArrayIndex, (int32)NumericValue);
+				return SetPropertyValue<UIntProperty, int32>(Property, Outer, Data, ArrayIndex, (int32)NumericValue);
 			}
-			else if (Property->GetClass() == UUInt32Property::StaticClass())
+			
+			if (Property->GetClass() == UUInt32Property::StaticClass())
 			{
-				SetPropertyValue<UUInt32Property, uint32>(Property, Outer, Data, ArrayIndex, (uint32)NumericValue);
+				return SetPropertyValue<UUInt32Property, uint32>(Property, Outer, Data, ArrayIndex, (uint32)NumericValue);
 			}
-			else if (Property->GetClass() == UInt16Property::StaticClass())
+			
+			if (Property->GetClass() == UInt16Property::StaticClass())
 			{
-				SetPropertyValue<UInt16Property, int16>(Property, Outer, Data, ArrayIndex, (int16)NumericValue);
+				return SetPropertyValue<UInt16Property, int16>(Property, Outer, Data, ArrayIndex, (int16)NumericValue);
 			}
-			else if (Property->GetClass() == UUInt16Property::StaticClass())
+			
+			if (Property->GetClass() == UUInt16Property::StaticClass())
 			{
-				SetPropertyValue<UUInt16Property, uint16>(Property, Outer, Data, ArrayIndex, (uint16)NumericValue);
+				return SetPropertyValue<UUInt16Property, uint16>(Property, Outer, Data, ArrayIndex, (uint16)NumericValue);
 			}
-			else if (Property->GetClass() == UInt64Property::StaticClass())
+			
+			if (Property->GetClass() == UInt64Property::StaticClass())
 			{
-				SetPropertyValue<UInt64Property, int64>(Property, Outer, Data, ArrayIndex, (int64)NumericValue);
+				return SetPropertyValue<UInt64Property, int64>(Property, Outer, Data, ArrayIndex, (int64)NumericValue);
 			}
-			else if (Property->GetClass() == UUInt64Property::StaticClass())
+			
+			if (Property->GetClass() == UUInt64Property::StaticClass())
 			{
-				SetPropertyValue<UUInt64Property, uint64>(Property, Outer, Data, ArrayIndex, (uint64)NumericValue);
+				return SetPropertyValue<UUInt64Property, uint64>(Property, Outer, Data, ArrayIndex, (uint64)NumericValue);
 			}
-			else if (Property->GetClass() == UInt8Property::StaticClass())
+			
+			if (Property->GetClass() == UInt8Property::StaticClass())
 			{
-				SetPropertyValue<UInt8Property, int8>(Property, Outer, Data, ArrayIndex, (int8)NumericValue);
+				return SetPropertyValue<UInt8Property, int8>(Property, Outer, Data, ArrayIndex, (int8)NumericValue);
 			}
-			else
-			{
-				UE_LOG(LogSerialization, Verbose, TEXT("Numeric field %s with value '%f' is not supported in UProperty type %s (%s)"), *Property->GetFName().ToString(), NumericValue, *Property->GetClass()->GetName(), *GetDebugString());
 
-				return false;
-			}
+			UE_LOG(LogSerialization, Verbose, TEXT("Numeric field %s with value '%f' is not supported in UProperty type %s (%s)"), *Property->GetFName().ToString(), NumericValue, *Property->GetClass()->GetName(), *GetDebugString());
+
+			return false;
 		}
 		break;
 
 	// null values
 	case EJsonNotation::Null:
-		ClearPropertyValue(Property, Outer, Data, ArrayIndex);
-		break;
+		return ClearPropertyValue(Property, Outer, Data, ArrayIndex);
 
 	// strings, names & enumerations
 	case EJsonNotation::String:
@@ -211,32 +253,35 @@ bool FJsonStructDeserializerBackend::ReadProperty( UProperty* Property, UPropert
 
 			if (Property->GetClass() == UStrProperty::StaticClass())
 			{
-				SetPropertyValue<UStrProperty, FString>(Property, Outer, Data, ArrayIndex, StringValue);
+				return SetPropertyValue<UStrProperty, FString>(Property, Outer, Data, ArrayIndex, StringValue);
 			}
-			else if (Property->GetClass() == UNameProperty::StaticClass())
+			
+			if (Property->GetClass() == UNameProperty::StaticClass())
 			{
-				SetPropertyValue<UNameProperty, FName>(Property, Outer, Data, ArrayIndex, *StringValue);
+				return SetPropertyValue<UNameProperty, FName>(Property, Outer, Data, ArrayIndex, *StringValue);
 			}
-			else if (Property->GetClass() == UByteProperty::StaticClass())
+			
+			if (Property->GetClass() == UByteProperty::StaticClass())
 			{
 				UByteProperty* ByteProperty = Cast<UByteProperty>(Property);
 				int32 Index = ByteProperty->Enum->FindEnumIndex(*StringValue);
 
-				if (Index != INDEX_NONE)
+				if (Index == INDEX_NONE)
 				{
-					SetPropertyValue<UByteProperty, uint8>(Property, Outer, Data, ArrayIndex, (uint8)Index);
+					return false;
 				}
+
+				return SetPropertyValue<UByteProperty, uint8>(Property, Outer, Data, ArrayIndex, (uint8)Index);
 			}
-			else if (Property->GetClass() == UClassProperty::StaticClass())
+			
+			if (Property->GetClass() == UClassProperty::StaticClass())
 			{
-				SetPropertyValue<UClassProperty, UClass*>(Property, Outer, Data, ArrayIndex, LoadObject<UClass>(NULL, *StringValue, NULL, LOAD_NoWarn));
+				return SetPropertyValue<UClassProperty, UClass*>(Property, Outer, Data, ArrayIndex, LoadObject<UClass>(NULL, *StringValue, NULL, LOAD_NoWarn));
 			}
-			else
-			{
-				UE_LOG(LogSerialization, Verbose, TEXT("String field %s with value '%s' is not supported in UProperty type %s (%s)"), *Property->GetFName().ToString(), *StringValue, *Property->GetClass()->GetName(), *GetDebugString());
 
-				return false;
-			}
+			UE_LOG(LogSerialization, Verbose, TEXT("String field %s with value '%s' is not supported in UProperty type %s (%s)"), *Property->GetFName().ToString(), *StringValue, *Property->GetClass()->GetName(), *GetDebugString());
+
+			return false;
 		}
 		break;
 	}
@@ -245,13 +290,13 @@ bool FJsonStructDeserializerBackend::ReadProperty( UProperty* Property, UPropert
 }
 
 
-void FJsonStructDeserializerBackend::SkipArray( ) 
+void FJsonStructDeserializerBackend::SkipArray() 
 {
 	JsonReader->SkipArray();
 }
 
 
-void FJsonStructDeserializerBackend::SkipStructure( )
+void FJsonStructDeserializerBackend::SkipStructure()
 {
 	JsonReader->SkipObject();
 }
