Serialization: Moved interfaces into public root

diff --git a/Engine/Source/Runtime/Serialization/Public/IStructDeserializerBackend.h b/Engine/Source/Runtime/Serialization/Public/IStructDeserializerBackend.h
new file mode 100644
--- /dev/null
+++ b/Engine/Source/Runtime/Serialization/Public/IStructDeserializerBackend.h
@@ -0,0 +1,102 @@
+// Copyright 1998-2013 Epic Games, Inc. All Rights Reserved.
+
+#pragma once
+
+
+class UProperty;
+class UStruct;
+
+
+/**
+ * Enumerates serialization tokens.
+ */
+enum class EStructDeserializerBackendTokens
+{
+	/** End of an array. */
+	ArrayEnd,
+
+	/** Beginning of an array. */
+	ArrayStart,
+
+	/** A comment. */
+	Comment,
+
+	/** An error occurred when reading the token. */
+	Error,
+
+	/** No token available. */
+	None,
+
+	/** A scalar property. */
+	Property,
+
+	/** End of a data structure. */
+	StructureEnd,
+
+	/** Beginning of a data structure. */
+	StructureStart,
+};
+
+
+/**
+ * Interface for UStruct serializer backends.
+ */
+class IStructDeserializerBackend
+{
+public:
+
+	/**
+	 * Gets the identifier of the current field.
+	 *
+	 * @return The field's name.
+	 */
+	virtual const FString& GetCurrentPropertyName() const = 0;
+
+	/**
+	 * Gets a debug string for the reader's current state.
+	 *
+	 * The returned string contains debug information that is relevant to the reader's
+	 * serialization format. For example, it could be a line and column number for text
+	 * based formats, or a byte offset for binary serialization formats.
+	 *
+	 * @return The debug string.
+	 */
+	virtual FString GetDebugString() const = 0;
+
+	/**
+	 * Gets the last error message.
+	 *
+	 * @return Error message.
+	 */
+	virtual const FString& GetLastErrorMessage() const = 0;
+
+	/**
+	 * Reads the next token from the stream.
+	 *
+	 * @param OutToken Will contain the read token, if any.
+	 * @return true if a token was read, false otherwise.
+	 */
+	virtual bool GetNextToken( EStructDeserializerBackendTokens& OutToken ) = 0;
+
+	/**
+	 * Reads the specified property from the stream.
+	 *
+	 * @param Property The property to read into.
+	 * @param Outer The outer property holding the property to read (in case of arrays).
+	 * @param Data The buffer that will hold the read data.
+	 * @param ArrayIndex - An index into the property array (for static arrays).
+	 * @return true on success, false otherwise.
+	 */
+	virtual bool ReadProperty( UProperty* Property, UProperty* Outer, void* Data, int32 ArrayIndex ) = 0;
+
+	/** Skips the array that is currently being read from the stream. */
+	virtual void SkipArray() = 0;
+
+	/** Skips the object that is currently being read from the stream. */
+	virtual void SkipStructure() = 0;
+
+public:
+
+	/** Virtual destructor. */
+	virtual ~IStructDeserializerBackend() { }
+};
diff --git a/Engine/Source/Runtime/Serialization/Public/IStructSerializerBackend.h b/Engine/Source/Runtime/Serialization/Public/IStructSerializerBackend.h
new file mode 100644
--- /dev/null
+++ b/Engine/Source/Runtime/Serialization/Public/IStructSerializerBackend.h
@@ -0,0 +1,73 @@
+// Copyright 1998-2013 Epic Games, Inc. All Rights Reserved.
+
+#pragma once
+
+
+class UProperty;
+class UStruct;
+
+
+/**
+ * Interface for UStruct serializer backends.
+ */
+class IStructSerializerBackend
+{
+public:
+
+	/**
+	 * Signals the beginning of an array.
+	 *
+	 * @param Property The property that holds the array.
+	 */
+	virtual void BeginArray( UProperty* Property ) = 0;
+
+	/**
+	 * Signals the beginning of a child structure.
+	 *
+	 * @param Property The property that holds the object.
+	 */
+	virtual void BeginStructure( UProperty* Property ) = 0;
+
+	/**
+	 * Signals the beginning of a root structure.
+	 *
+	 * @param TypeInfo The object's type information.
+	 */
+	virtual void BeginStructure( UStruct* TypeInfo ) = 0;
+
+	/**
+	 * Signals the end of an array.
+	 *
+	 * @param Property The property that holds the array.
+	 */
+	virtual void EndArray( UProperty* Property ) = 0;
+
+	/**
+	 * Signals the end of an object.
+	 */
+	virtual void EndStructure() = 0;
+
+	/**
+	 * Writes a comment to the output stream.
+	 *
+	 * @param Comment The comment text.
+	 */
+	virtual void WriteComment( const FString& Comment ) = 0;
+
+	/**
+	 * Writes a property to the output stream.
+	 *
+	 * Depending on the context, properties to be written can be either object properties or array elements.
+	 *
+	 * @param Property The property that holds the data to write.
+	 * @param Data The property's data to write.
+	 * @param TypeInfo The property data's type information.
+	 * @param ArrayIndex An index into the property array (for static arrays).
+	 */
+	virtual void WriteProperty( UProperty* Property, const void* Data, UStruct* TypeInfo, int32 ArrayIndex = 0 ) = 0;
+
+public:
+
+	/** Virtual destructor. */
+	virtual ~IStructSerializerBackend() { }
+};
diff --git a/Engine/Source/Runtime/Serialization/Public/Interfaces/IStructDeserializerBackend.h b/Engine/Source/Runtime/Serialization/Public/Interfaces/IStructDeserializerBackend.h
deleted file mode 100644
diff --git a/Engine/Source/Runtime/Serialization/Public/Interfaces/IStructSerializerBackend.h b/Engine/Source/Runtime/Serialization/Public/Interfaces/IStructSerializerBackend.h
deleted file mode 100644
