Serialization: Improved build rules and forward declarations; added documentation

diff --git a/Engine/Source/Runtime/Serialization/Private/SerializationModule.cpp b/Engine/Source/Runtime/Serialization/Private/SerializationModule.cpp
--- a/Engine/Source/Runtime/Serialization/Private/SerializationModule.cpp
+++ b/Engine/Source/Runtime/Serialization/Private/SerializationModule.cpp
@@ -1,7 +1,7 @@
 // Copyright 1998-2013 Epic Games, Inc. All Rights Reserved.
 
 #include "SerializationPrivatePCH.h"
-#include "ModuleManager.h"
+#include "ModuleInterface.h"
 
 
 //DEFINE_LOG_CATEGORY(LogSerialization);
diff --git a/Engine/Source/Runtime/Serialization/Private/StructDeserializer.cpp b/Engine/Source/Runtime/Serialization/Private/StructDeserializer.cpp
--- a/Engine/Source/Runtime/Serialization/Private/StructDeserializer.cpp
+++ b/Engine/Source/Runtime/Serialization/Private/StructDeserializer.cpp
@@ -10,24 +10,31 @@
 
 namespace StructDeserializer
 {
-	// Structure for the read state stack.
+	/**
+	 * Structure for the read state stack.
+	 */
 	struct FReadState
 	{
-		// Holds the property's current array index.
+		/** Holds the property's current array index. */
 		int32 ArrayIndex;
 
-		// Holds a pointer to the property's data.
+		/** Holds a pointer to the property's data. */
 		void* Data;
 
-		// Holds the property's meta data.
+		/** Holds the property's meta data. */
 		UProperty* Property;
 
-		// Holds a pointer to the UStruct describing the data.
+		/** Holds a pointer to the UStruct describing the data. */
 		UStruct* TypeInfo;
 	};
 
 
-	// Finds the class for the given stack state.
+	/**
+	 * Finds the class for the given stack state.
+	 *
+	 * @param State The stack state to find the class for.
+	 * @return The class, if found.
+	 */
 	UStruct* FindClass( const FReadState& State )
 	{
 		UStruct* Class = nullptr;
diff --git a/Engine/Source/Runtime/Serialization/Private/StructSerializer.cpp b/Engine/Source/Runtime/Serialization/Private/StructSerializer.cpp
--- a/Engine/Source/Runtime/Serialization/Private/StructSerializer.cpp
+++ b/Engine/Source/Runtime/Serialization/Private/StructSerializer.cpp
@@ -10,24 +10,33 @@
 
 namespace StructSerializer
 {
-	// Structure for the write state stack.
+	/**
+	 * Structure for the write state stack.
+	 */
 	struct FWriteState
 	{
-		// Holds a pointer to the property's data.
+		/** Holds a pointer to the property's data. */
 		const void* Data;
 
-		// Holds a flag indicating whether the property has been processed.
+		/** Holds a flag indicating whether the property has been processed. */
 		bool HasBeenProcessed;
 
-		// Holds the property's meta data.
+		/** Holds the property's meta data. */
 		UProperty* Property;
 
-		// Holds a pointer to the UStruct describing the data.
+		/** Holds a pointer to the UStruct describing the data. */
 		UStruct* TypeInfo;
 	};
 
 
-	// Gets the value from the given property.
+	/**
+	 * Gets the value from the given property.
+	 *
+	 * @param PropertyType The type name of the property.
+	 * @param State The stack state that holds the property's data.
+	 * @param Property A pointer to the property.
+	 * @return A pointer to the property's value, or nullptr if it couldn't be found.
+	 */
 	template<typename UPropertyType, typename PropertyType>
 	PropertyType* GetPropertyValue( const FWriteState& State, UProperty* Property )
 	{
diff --git a/Engine/Source/Runtime/Serialization/Public/Backends/JsonStructDeserializerBackend.h b/Engine/Source/Runtime/Serialization/Public/Backends/JsonStructDeserializerBackend.h
--- a/Engine/Source/Runtime/Serialization/Public/Backends/JsonStructDeserializerBackend.h
+++ b/Engine/Source/Runtime/Serialization/Public/Backends/JsonStructDeserializerBackend.h
@@ -6,6 +6,10 @@
 #include "Json.h"
 
 
+// forward declarations
+class UProperty;
+
+
 /**
  * Implements a reader for UStruct deserialization using Json.
  *
diff --git a/Engine/Source/Runtime/Serialization/Public/Backends/JsonStructSerializerBackend.h b/Engine/Source/Runtime/Serialization/Public/Backends/JsonStructSerializerBackend.h
--- a/Engine/Source/Runtime/Serialization/Public/Backends/JsonStructSerializerBackend.h
+++ b/Engine/Source/Runtime/Serialization/Public/Backends/JsonStructSerializerBackend.h
@@ -6,6 +6,11 @@
 #include "Json.h"
 
 
+// forward declarations
+class UProperty;
+class UStruct;
+
+
 /**
  * Implements a writer for UStruct serialization using Json.
  *
diff --git a/Engine/Source/Runtime/Serialization/Public/IStructDeserializerBackend.h b/Engine/Source/Runtime/Serialization/Public/IStructDeserializerBackend.h
--- a/Engine/Source/Runtime/Serialization/Public/IStructDeserializerBackend.h
+++ b/Engine/Source/Runtime/Serialization/Public/IStructDeserializerBackend.h
@@ -3,8 +3,8 @@
 #pragma once
 
 
+// forward declarations
 class UProperty;
-class UStruct;
 
 
 /**
diff --git a/Engine/Source/Runtime/Serialization/Public/IStructSerializerBackend.h b/Engine/Source/Runtime/Serialization/Public/IStructSerializerBackend.h
--- a/Engine/Source/Runtime/Serialization/Public/IStructSerializerBackend.h
+++ b/Engine/Source/Runtime/Serialization/Public/IStructSerializerBackend.h
@@ -3,6 +3,7 @@
 #pragma once
 
 
+// forward declarations
 class UProperty;
 class UStruct;
 
diff --git a/Engine/Source/Runtime/Serialization/Public/StructDeserializer.h b/Engine/Source/Runtime/Serialization/Public/StructDeserializer.h
--- a/Engine/Source/Runtime/Serialization/Public/StructDeserializer.h
+++ b/Engine/Source/Runtime/Serialization/Public/StructDeserializer.h
@@ -3,6 +3,7 @@
 #pragma once
 
 
+// forward declarations
 class IStructDeserializerBackend;
 class UStruct;
 
diff --git a/Engine/Source/Runtime/Serialization/Public/StructSerializer.h b/Engine/Source/Runtime/Serialization/Public/StructSerializer.h
--- a/Engine/Source/Runtime/Serialization/Public/StructSerializer.h
+++ b/Engine/Source/Runtime/Serialization/Public/StructSerializer.h
@@ -3,6 +3,7 @@
 #pragma once
 
 
+// forward declarations
 class IStructSerializerBackend;
 class UStruct;
 
@@ -41,7 +42,7 @@ enum class EStructSerializerReferenceLoopPolicies
  */
 struct FStructSerializerPolicies
 {
-	/** Holds the policy for NULL values. */
+	/** Holds the policy for null values. */
 	EStructSerializerNullValuePolicies NullValues;
 
 	/** Holds the policy for reference loops. */
diff --git a/Engine/Source/Runtime/Serialization/Serialization.Build.cs b/Engine/Source/Runtime/Serialization/Serialization.Build.cs
--- a/Engine/Source/Runtime/Serialization/Serialization.Build.cs
+++ b/Engine/Source/Runtime/Serialization/Serialization.Build.cs
@@ -5,16 +5,22 @@ namespace UnrealBuildTool.Rules
 	public class Serialization : ModuleRules
 	{
 		public Serialization(TargetInfo Target)
-		{	
-			PrivateDependencyModuleNames.AddRange(
+		{
+			PublicDependencyModuleNames.AddRange(
 				new string[]
 				{
 					"Core",
-					"CoreUObject",
 					"Json",
 				}
 			);
 
+			PrivateDependencyModuleNames.AddRange(
+				new string[]
+				{
+					"CoreUObject",
+				}
+			);
+
 			PrivateIncludePaths.AddRange(
 				new string[] {
 					"Runtime/Serialization/Private",
