Serialization: Fixed incorrect handling of large integers in Json backends

diff --git a/Engine/Source/Runtime/Serialization/Private/Backends/JsonStructDeserializerBackend.cpp b/Engine/Source/Runtime/Serialization/Private/Backends/JsonStructDeserializerBackend.cpp
--- a/Engine/Source/Runtime/Serialization/Private/Backends/JsonStructDeserializerBackend.cpp
+++ b/Engine/Source/Runtime/Serialization/Private/Backends/JsonStructDeserializerBackend.cpp
@@ -184,7 +184,7 @@ bool FJsonStructDeserializerBackend::ReadProperty( UProperty* Property, UPropert
 	// numeric values
 	case EJsonNotation::Number:
 		{
-			float NumericValue = JsonReader->GetValueAsNumber();
+			double NumericValue = JsonReader->GetValueAsNumber();
 
 			if (Property->GetClass() == UByteProperty::StaticClass())
 			{
diff --git a/Engine/Source/Runtime/Serialization/Private/Backends/JsonStructSerializerBackend.cpp b/Engine/Source/Runtime/Serialization/Private/Backends/JsonStructSerializerBackend.cpp
--- a/Engine/Source/Runtime/Serialization/Private/Backends/JsonStructSerializerBackend.cpp
+++ b/Engine/Source/Runtime/Serialization/Private/Backends/JsonStructSerializerBackend.cpp
@@ -116,7 +116,7 @@ void FJsonStructSerializerBackend::WriteProperty( UProperty* Property, const voi
 		}
 		else
 		{
-			WritePropertyValue(JsonWriter, Property, (float)Cast<UByteProperty>(Property)->GetPropertyValue_InContainer(Data, ArrayIndex));
+			WritePropertyValue(JsonWriter, Property, (double)Cast<UByteProperty>(Property)->GetPropertyValue_InContainer(Data, ArrayIndex));
 		}			
 	}
 
@@ -133,33 +133,33 @@ void FJsonStructSerializerBackend::WriteProperty( UProperty* Property, const voi
 	// signed integers
 	else if (TypeInfo == UIntProperty::StaticClass())
 	{
-		WritePropertyValue(JsonWriter, Property, (float)Cast<UIntProperty>(Property)->GetPropertyValue_InContainer(Data, ArrayIndex));
+		WritePropertyValue(JsonWriter, Property, (double)Cast<UIntProperty>(Property)->GetPropertyValue_InContainer(Data, ArrayIndex));
 	}
 	else if (TypeInfo == UInt8Property::StaticClass())
 	{
-		WritePropertyValue(JsonWriter, Property, (float)Cast<UInt8Property>(Property)->GetPropertyValue_InContainer(Data, ArrayIndex));
+		WritePropertyValue(JsonWriter, Property, (double)Cast<UInt8Property>(Property)->GetPropertyValue_InContainer(Data, ArrayIndex));
 	}
 	else if (TypeInfo == UInt16Property::StaticClass())
 	{
-		WritePropertyValue(JsonWriter, Property, (float)Cast<UInt16Property>(Property)->GetPropertyValue_InContainer(Data, ArrayIndex));
+		WritePropertyValue(JsonWriter, Property, (double)Cast<UInt16Property>(Property)->GetPropertyValue_InContainer(Data, ArrayIndex));
 	}
 	else if (TypeInfo == UInt64Property::StaticClass())
 	{
-		WritePropertyValue(JsonWriter, Property, (float)Cast<UInt64Property>(Property)->GetPropertyValue_InContainer(Data, ArrayIndex));
+		WritePropertyValue(JsonWriter, Property, (double)Cast<UInt64Property>(Property)->GetPropertyValue_InContainer(Data, ArrayIndex));
 	}
 
 	// unsigned integers
 	else if (TypeInfo == UUInt16Property::StaticClass())
 	{
-		WritePropertyValue(JsonWriter, Property, (float)Cast<UUInt16Property>(Property)->GetPropertyValue_InContainer(Data, ArrayIndex));
+		WritePropertyValue(JsonWriter, Property, (double)Cast<UUInt16Property>(Property)->GetPropertyValue_InContainer(Data, ArrayIndex));
 	}
 	else if (TypeInfo == UUInt32Property::StaticClass())
 	{
-		WritePropertyValue(JsonWriter, Property, (float)Cast<UUInt32Property>(Property)->GetPropertyValue_InContainer(Data, ArrayIndex));
+		WritePropertyValue(JsonWriter, Property, (double)Cast<UUInt32Property>(Property)->GetPropertyValue_InContainer(Data, ArrayIndex));
 	}
 	else if (TypeInfo == UUInt64Property::StaticClass())
 	{
-		WritePropertyValue(JsonWriter, Property, (float)Cast<UUInt64Property>(Property)->GetPropertyValue_InContainer(Data, ArrayIndex));
+		WritePropertyValue(JsonWriter, Property, (double)Cast<UUInt64Property>(Property)->GetPropertyValue_InContainer(Data, ArrayIndex));
 	}
 
 	// names & strings
