Serialization: Improved unit tests for numeric values and FGuid

diff --git a/Engine/Source/Runtime/Serialization/Private/Tests/StructSerializerTest.cpp b/Engine/Source/Runtime/Serialization/Private/Tests/StructSerializerTest.cpp
--- a/Engine/Source/Runtime/Serialization/Private/Tests/StructSerializerTest.cpp
+++ b/Engine/Source/Runtime/Serialization/Private/Tests/StructSerializerTest.cpp
@@ -50,6 +50,7 @@ namespace StructSerializerTest
 		Test.TestEqual<TSubclassOf<class UObject>>(TEXT("Objects.Class must be the same before and after de-/serialization"), TestStruct.Objects.Class, TestStruct2.Objects.Class);
 		Test.TestEqual<UObject*>(TEXT("Objects.ObjectPtr must be the same before and after de-/serialization"), TestStruct.Objects.ObjectPtr, TestStruct2.Objects.ObjectPtr);
 
+		Test.TestEqual<FGuid>(TEXT("Builtins.Guid must be the same before and after de-/serialization"), TestStruct.Builtins.Guid, TestStruct2.Builtins.Guid);
 		Test.TestEqual<FName>(TEXT("Builtins.Name must be the same before and after de-/serialization"), TestStruct.Builtins.Name, TestStruct2.Builtins.Name);
 		Test.TestEqual<FString>(TEXT("Builtins.String must be the same before and after de-/serialization"), TestStruct.Builtins.String, TestStruct2.Builtins.String);
 		Test.TestEqual<FRotator>(TEXT("Builtins.Rotator must be the same before and after de-/serialization"), TestStruct.Builtins.Rotator, TestStruct2.Builtins.Rotator);
@@ -86,6 +87,10 @@ bool FJsonStructSerializerTest::RunTest( const FString& Parameters )
 		FJsonStructDeserializerBackend DeserializerBackend(Reader);
 
 		StructSerializerTest::TestSerialization(*this, SerializerBackend, DeserializerBackend);
+
+		// uncomment this to look at the serialized data
+		//uint8* Data = Buffer.GetData();
+		//GLog->Logf(TEXT("%s"), (TCHAR*)Data);
 	}
 
 	return true;
diff --git a/Engine/Source/Runtime/Serialization/Private/Tests/StructSerializerTestTypes.h b/Engine/Source/Runtime/Serialization/Private/Tests/StructSerializerTestTypes.h
--- a/Engine/Source/Runtime/Serialization/Private/Tests/StructSerializerTestTypes.h
+++ b/Engine/Source/Runtime/Serialization/Private/Tests/StructSerializerTestTypes.h
@@ -43,14 +43,14 @@ struct FStructSerializerNumericTestStruct
 
 	/** Default constructor. */
 	FStructSerializerNumericTestStruct()
-		: Int8(-8)
-		, Int16(-16)
-		, Int32(-32)
-		, Int64(-64)
-		, UInt8(8)
-		, UInt16(16)
-		, UInt32(32)
-		, UInt64(64)
+		: Int8(-127)
+		, Int16(-32767)
+		, Int32(-2147483647)
+		, Int64(-92233720368547/*75807*/)
+		, UInt8(255)
+		, UInt16(65535)
+		, UInt32(4294967295)
+		, UInt64(18446744073709/*551615*/)
 		, Float(4.125)
 		, Double(1.03125)
 	{ }
@@ -121,6 +121,9 @@ struct FStructSerializerBuiltinTestStruct
 {
 	GENERATED_USTRUCT_BODY()
 
+	UPROPERTY()
+	FGuid Guid;
+
 	UPROPERTY()
 	FName Name;
 
@@ -135,7 +138,8 @@ struct FStructSerializerBuiltinTestStruct
 
 	/** Default constructor. */
 	FStructSerializerBuiltinTestStruct()
-		: String("Hello World!")
+		: Guid(FGuid::NewGuid())
+		, String("Hello World!")
 		, Rotator(4096, 8192, 16384)
 		, Vector(1.0f, 2.0f, 3.0f)
 	{ }
