Fixed comment.

diff --git a/Engine/Source/Runtime/Serialization/Public/IStructDeserializerBackend.h b/Engine/Source/Runtime/Serialization/Public/IStructDeserializerBackend.h
--- a/Engine/Source/Runtime/Serialization/Public/IStructDeserializerBackend.h
+++ b/Engine/Source/Runtime/Serialization/Public/IStructDeserializerBackend.h
@@ -84,7 +84,7 @@ public:
 	 * @param Property The property to read into.
 	 * @param Outer The outer property holding the property to read (in case of arrays).
 	 * @param Data The buffer that will hold the read data.
-	 * @param ArrayIndex - An index into the property array (for static arrays).
+	 * @param ArrayIndex An index into the property array (for static arrays).
 	 * @return true on success, false otherwise.
 	 */
 	virtual bool ReadProperty( UProperty* Property, UProperty* Outer, void* Data, int32 ArrayIndex ) = 0;
