Serialization: Added support for FText properties

diff --git a/Engine/Source/Runtime/Serialization/Private/Backends/JsonStructDeserializerBackend.cpp b/Engine/Source/Runtime/Serialization/Private/Backends/JsonStructDeserializerBackend.cpp
--- a/Engine/Source/Runtime/Serialization/Private/Backends/JsonStructDeserializerBackend.cpp
+++ b/Engine/Source/Runtime/Serialization/Private/Backends/JsonStructDeserializerBackend.cpp
@@ -255,12 +255,17 @@ bool FJsonStructDeserializerBackend::ReadProperty( UProperty* Property, UPropert
 			{
 				return SetPropertyValue<UStrProperty, FString>(Property, Outer, Data, ArrayIndex, StringValue);
 			}
-			
+		
 			if (Property->GetClass() == UNameProperty::StaticClass())
 			{
 				return SetPropertyValue<UNameProperty, FName>(Property, Outer, Data, ArrayIndex, *StringValue);
 			}
 			
+			if (Property->GetClass() == UTextProperty::StaticClass())
+			{
+				return SetPropertyValue<UTextProperty, FText>(Property, Outer, Data, ArrayIndex, FText::FromString(StringValue));
+			}
+
 			if (Property->GetClass() == UByteProperty::StaticClass())
 			{
 				UByteProperty* ByteProperty = Cast<UByteProperty>(Property);
diff --git a/Engine/Source/Runtime/Serialization/Private/Backends/JsonStructSerializerBackend.cpp b/Engine/Source/Runtime/Serialization/Private/Backends/JsonStructSerializerBackend.cpp
--- a/Engine/Source/Runtime/Serialization/Private/Backends/JsonStructSerializerBackend.cpp
+++ b/Engine/Source/Runtime/Serialization/Private/Backends/JsonStructSerializerBackend.cpp
@@ -162,7 +162,7 @@ void FJsonStructSerializerBackend::WriteProperty( UProperty* Property, const voi
 		WritePropertyValue(JsonWriter, Property, (double)Cast<UUInt64Property>(Property)->GetPropertyValue_InContainer(Data, ArrayIndex));
 	}
 
-	// names & strings
+	// names, strings & text
 	else if (TypeInfo == UNameProperty::StaticClass())
 	{
 		WritePropertyValue(JsonWriter, Property, Cast<UNameProperty>(Property)->GetPropertyValue_InContainer(Data, ArrayIndex).ToString());
@@ -171,6 +171,10 @@ void FJsonStructSerializerBackend::WriteProperty( UProperty* Property, const voi
 	{
 		WritePropertyValue(JsonWriter, Property, Cast<UStrProperty>(Property)->GetPropertyValue_InContainer(Data, ArrayIndex));
 	}
+	else if (TypeInfo == UTextProperty::StaticClass())
+	{
+		WritePropertyValue(JsonWriter, Property, Cast<UTextProperty>(Property)->GetPropertyValue_InContainer(Data, ArrayIndex).ToString());
+	}
 
 	// classes & objects
 	else if (TypeInfo == UClassProperty::StaticClass())
diff --git a/Engine/Source/Runtime/Serialization/Private/Tests/StructSerializerTest.cpp b/Engine/Source/Runtime/Serialization/Private/Tests/StructSerializerTest.cpp
--- a/Engine/Source/Runtime/Serialization/Private/Tests/StructSerializerTest.cpp
+++ b/Engine/Source/Runtime/Serialization/Private/Tests/StructSerializerTest.cpp
@@ -54,6 +54,7 @@ namespace StructSerializerTest
 		Test.TestEqual<FName>(TEXT("Builtins.Name must be the same before and after de-/serialization"), TestStruct.Builtins.Name, TestStruct2.Builtins.Name);
 		Test.TestEqual<FString>(TEXT("Builtins.String must be the same before and after de-/serialization"), TestStruct.Builtins.String, TestStruct2.Builtins.String);
 		Test.TestEqual<FRotator>(TEXT("Builtins.Rotator must be the same before and after de-/serialization"), TestStruct.Builtins.Rotator, TestStruct2.Builtins.Rotator);
+		Test.TestEqual<FString>(TEXT("Builtins.Text must be the same before and after de-/serialization"), TestStruct.Builtins.Text.ToString(), TestStruct2.Builtins.Text.ToString());
 		Test.TestEqual<FVector>(TEXT("Builtins.Vector must be the same before and after de-/serialization"), TestStruct.Builtins.Vector, TestStruct2.Builtins.Vector);
 
 		Test.TestEqual<TArray<int32>>(TEXT("Arrays.Int32Array must be the same before and after de-/serialization"), TestStruct.Arrays.Int32Array, TestStruct2.Arrays.Int32Array);
diff --git a/Engine/Source/Runtime/Serialization/Private/Tests/StructSerializerTestTypes.h b/Engine/Source/Runtime/Serialization/Private/Tests/StructSerializerTestTypes.h
--- a/Engine/Source/Runtime/Serialization/Private/Tests/StructSerializerTestTypes.h
+++ b/Engine/Source/Runtime/Serialization/Private/Tests/StructSerializerTestTypes.h
@@ -134,14 +134,18 @@ struct FStructSerializerBuiltinTestStruct
 	UPROPERTY()
 	FRotator Rotator;
 
+	UPROPERTY()
+	FText Text;
+
 	UPROPERTY()
 	FVector Vector;
 
 	/** Default constructor. */
 	FStructSerializerBuiltinTestStruct()
 		: Guid(FGuid::NewGuid())
-		, String("Hello World!")
+		, String("Test String")
 		, Rotator(4096, 8192, 16384)
+		, Text(FText::FromString("Test Text"))
 		, Vector(1.0f, 2.0f, 3.0f)
 	{ }
 
