[INTEGRATE] Add ways to filter properties for (de)serialization
Merging CL#2624604 using UE4-To-UE4-LauncherDev

diff --git a/Engine/Source/Runtime/Serialization/Private/StructDeserializer.cpp b/Engine/Source/Runtime/Serialization/Private/StructDeserializer.cpp
--- a/Engine/Source/Runtime/Serialization/Private/StructDeserializer.cpp
+++ b/Engine/Source/Runtime/Serialization/Private/StructDeserializer.cpp
@@ -135,6 +135,11 @@ bool FStructDeserializer::Deserialize( void* OutStruct, UStruct& TypeInfo, IStru
 				}
 				else
 				{
+					// Skip array property if property filter is set and rejects it
+					if (Policies.PropertyFilter && Policies.PropertyFilter(NewState.Property, CurrentState.Property))
+					{
+						continue;
+					}
 					NewState.ArrayIndex = 0;
 					NewState.Data = CurrentState.Data;
 					NewState.TypeInfo = FindClass(NewState);
@@ -205,9 +210,18 @@ bool FStructDeserializer::Deserialize( void* OutStruct, UStruct& TypeInfo, IStru
 							return false;
 						}
 					}
-					else if (!Backend.ReadProperty(Property, CurrentState.Property, CurrentState.Data, CurrentState.ArrayIndex))
+					else
 					{
-						UE_LOG(LogSerialization, Verbose, TEXT("The property '%s' could not be read (%s)"), *PropertyName, *Backend.GetDebugString());
+						// Skip scalar property if property filter is set and rejects it
+						if (Policies.PropertyFilter && Policies.PropertyFilter(Property, CurrentState.Property))
+						{
+							continue;
+						}
+
+						if (!Backend.ReadProperty(Property, CurrentState.Property, CurrentState.Data, CurrentState.ArrayIndex))
+						{
+							UE_LOG(LogSerialization, Verbose, TEXT("The property '%s' could not be read (%s)"), *PropertyName, *Backend.GetDebugString());
+						}
 					}
 				}
 			}
@@ -278,6 +292,13 @@ bool FStructDeserializer::Deserialize( void* OutStruct, UStruct& TypeInfo, IStru
 
 				if (NewState.Property != nullptr)
 				{
+					// Skip struct property if property filter is set and rejects it
+					if (Policies.PropertyFilter && !Policies.PropertyFilter(NewState.Property, CurrentState.Property))
+					{
+						Backend.SkipStructure();
+						continue;
+					}
+
 					NewState.ArrayIndex = 0;
 					NewState.TypeInfo = FindClass(NewState);
 
diff --git a/Engine/Source/Runtime/Serialization/Private/StructSerializer.cpp b/Engine/Source/Runtime/Serialization/Private/StructSerializer.cpp
--- a/Engine/Source/Runtime/Serialization/Private/StructSerializer.cpp
+++ b/Engine/Source/Runtime/Serialization/Private/StructSerializer.cpp
@@ -143,6 +143,11 @@ void FStructSerializer::Serialize( const void* Struct, UStruct& TypeInfo, IStruc
 
 				for (TFieldIterator<UProperty> It(CurrentState.TypeInfo, EFieldIteratorFlags::IncludeSuper); It; ++It)
 				{
+					// Skip property if the filter function is set and rejects it.
+					if (Policies.PropertyFilter && !Policies.PropertyFilter(*It, CurrentState.Property))
+					{
+						continue;
+					}
 					FWriteState NewState;
 
 					NewState.Data = NewData;
diff --git a/Engine/Source/Runtime/Serialization/Public/StructDeserializer.h b/Engine/Source/Runtime/Serialization/Public/StructDeserializer.h
--- a/Engine/Source/Runtime/Serialization/Public/StructDeserializer.h
+++ b/Engine/Source/Runtime/Serialization/Public/StructDeserializer.h
@@ -32,9 +32,15 @@ struct FStructDeserializerPolicies
 	/** Holds the policy for handling missing fields. */
 	EStructDeserializerErrorPolicies MissingFields;
 
+	/** Predicate for performing advanced filtering of struct properties. 
+		If set, the predicate should return true for all properties it wishes to include in the output.
+	 */
+	TFunction<bool (const UProperty* /*CurrentProp*/, const UProperty* /*ParentProp*/)> PropertyFilter;
+
 	/** Default constructor. */
 	FStructDeserializerPolicies()
 		: MissingFields(EStructDeserializerErrorPolicies::Ignore)
+		, PropertyFilter()
 	{ }
 };
 
diff --git a/Engine/Source/Runtime/Serialization/Public/StructSerializer.h b/Engine/Source/Runtime/Serialization/Public/StructSerializer.h
--- a/Engine/Source/Runtime/Serialization/Public/StructSerializer.h
+++ b/Engine/Source/Runtime/Serialization/Public/StructSerializer.h
@@ -48,10 +48,16 @@ struct FStructSerializerPolicies
 	/** Holds the policy for reference loops. */
 	EStructSerializerReferenceLoopPolicies ReferenceLoops;
 
+	/** Predicate for performing advanced filtering of struct properties. 
+		If set, the predicate should return true for all properties it wishes to include in the output.
+	 */
+	TFunction<bool (const UProperty* /*CurrentProp*/, const UProperty* /*ParentProp*/)> PropertyFilter;
+
 	/** Default constructor. */
 	FStructSerializerPolicies()
 		: NullValues(EStructSerializerNullValuePolicies::Serialize)
 		, ReferenceLoops(EStructSerializerReferenceLoopPolicies::Ignore)
+		, PropertyFilter()
 	{ }
 };
 
