Serialization: Exposed entire serializer state in backend API

#UpgradeNotes: This is an API breaking change. See the documentation in IStructSerializerBackend and the changes in FJsonStructSerializerBackend and WebJSStructSerializerBackend for details.

diff --git a/Engine/Source/Runtime/Serialization/Private/Backends/JsonStructSerializerBackend.cpp b/Engine/Source/Runtime/Serialization/Private/Backends/JsonStructSerializerBackend.cpp
--- a/Engine/Source/Runtime/Serialization/Private/Backends/JsonStructSerializerBackend.cpp
+++ b/Engine/Source/Runtime/Serialization/Private/Backends/JsonStructSerializerBackend.cpp
@@ -41,9 +41,9 @@ namespace JsonStructSerializerBackend
 /* IStructSerializerBackend interface
  *****************************************************************************/
 
-void FJsonStructSerializerBackend::BeginArray( UProperty* Property )
+void FJsonStructSerializerBackend::BeginArray(const FStructSerializerState& State)
 {
-	UObject* Outer = Property->GetOuter();
+	UObject* Outer = State.ValueProperty->GetOuter();
 
 	if ((Outer != nullptr) && (Outer->GetClass() == UArrayProperty::StaticClass()))
 	{
@@ -51,143 +51,145 @@ void FJsonStructSerializerBackend::BeginArray( UProperty* Property )
 	}
 	else
 	{
-		JsonWriter->WriteArrayStart(Property->GetName());
+		JsonWriter->WriteArrayStart(State.ValueProperty->GetName());
 	}
 }
 
 
-void FJsonStructSerializerBackend::BeginStructure( UProperty* Property )
+void FJsonStructSerializerBackend::BeginStructure(const FStructSerializerState& State)
 {
-	UObject* Outer = Property->GetOuter();
-
-	if ((Outer != nullptr) && (Outer->GetClass() == UArrayProperty::StaticClass()))
+	if (State.ValueProperty != nullptr)
 	{
-		JsonWriter->WriteObjectStart();
+		UObject* Outer = State.ValueProperty->GetOuter();
+
+		if ((Outer != nullptr) && (Outer->GetClass() == UArrayProperty::StaticClass()))
+		{
+			JsonWriter->WriteObjectStart();
+		}
+		else
+		{
+			JsonWriter->WriteObjectStart(State.ValueProperty->GetName());
+		}
 	}
 	else
 	{
-		JsonWriter->WriteObjectStart(Property->GetName());
+		JsonWriter->WriteObjectStart();
 	}
 }
 
 
-void FJsonStructSerializerBackend::BeginStructure( UStruct* TypeInfo )
-{
-	JsonWriter->WriteObjectStart();
-}
-
-
-void FJsonStructSerializerBackend::EndArray( UProperty* Property )
+void FJsonStructSerializerBackend::EndArray(const FStructSerializerState& /*State*/)
 {
 	JsonWriter->WriteArrayEnd();
 }
 
 
-void FJsonStructSerializerBackend::EndStructure()
+void FJsonStructSerializerBackend::EndStructure(const FStructSerializerState& /*State*/)
 {
 	JsonWriter->WriteObjectEnd();
 }
 
 
-void FJsonStructSerializerBackend::WriteComment( const FString& Comment )
+void FJsonStructSerializerBackend::WriteComment(const FString& Comment)
 {
 	// Json does not support comments
 }
 
 
-void FJsonStructSerializerBackend::WriteProperty( UProperty* Property, const void* Data, UStruct* TypeInfo, int32 ArrayIndex )
+void FJsonStructSerializerBackend::WriteProperty(const FStructSerializerState& State, int32 ArrayIndex)
 {
 	using namespace JsonStructSerializerBackend;
 
 	// booleans
-	if (TypeInfo == UBoolProperty::StaticClass())
+	if (State.ValueType == UBoolProperty::StaticClass())
 	{
-		WritePropertyValue(JsonWriter, Property, Cast<UBoolProperty>(Property)->GetPropertyValue_InContainer(Data, ArrayIndex));
+		WritePropertyValue(JsonWriter, State.ValueProperty, Cast<UBoolProperty>(State.ValueProperty)->GetPropertyValue_InContainer(State.ValueData, ArrayIndex));
 	}
 
 	// unsigned bytes & enumerations
-	else if (TypeInfo == UByteProperty::StaticClass())
+	else if (State.ValueType == UByteProperty::StaticClass())
 	{
-		UByteProperty* ByteProperty = Cast<UByteProperty>(Property);
+		UByteProperty* ByteProperty = Cast<UByteProperty>(State.ValueProperty);
 
 		if (ByteProperty->IsEnum())
 		{
-			WritePropertyValue(JsonWriter, Property, ByteProperty->Enum->GetEnumName(ByteProperty->GetPropertyValue_InContainer(Data, ArrayIndex)));
+			WritePropertyValue(JsonWriter, State.ValueProperty, ByteProperty->Enum->GetEnumName(ByteProperty->GetPropertyValue_InContainer(State.ValueData, ArrayIndex)));
 		}
 		else
 		{
-			WritePropertyValue(JsonWriter, Property, (double)Cast<UByteProperty>(Property)->GetPropertyValue_InContainer(Data, ArrayIndex));
+			WritePropertyValue(JsonWriter, State.ValueProperty, (double)Cast<UByteProperty>(State.ValueProperty)->GetPropertyValue_InContainer(State.ValueData, ArrayIndex));
 		}			
 	}
 
 	// floating point numbers
-	else if (TypeInfo == UDoubleProperty::StaticClass())
+	else if (State.ValueType == UDoubleProperty::StaticClass())
 	{
-		WritePropertyValue(JsonWriter, Property, Cast<UDoubleProperty>(Property)->GetPropertyValue_InContainer(Data, ArrayIndex));
+		WritePropertyValue(JsonWriter, State.ValueProperty, Cast<UDoubleProperty>(State.ValueProperty)->GetPropertyValue_InContainer(State.ValueData, ArrayIndex));
 	}
-	else if (TypeInfo == UFloatProperty::StaticClass())
+	else if (State.ValueType == UFloatProperty::StaticClass())
 	{
-		WritePropertyValue(JsonWriter, Property, Cast<UFloatProperty>(Property)->GetPropertyValue_InContainer(Data, ArrayIndex));
+		WritePropertyValue(JsonWriter, State.ValueProperty, Cast<UFloatProperty>(State.ValueProperty)->GetPropertyValue_InContainer(State.ValueData, ArrayIndex));
 	}
 
 	// signed integers
-	else if (TypeInfo == UIntProperty::StaticClass())
+	else if (State.ValueType == UIntProperty::StaticClass())
 	{
-		WritePropertyValue(JsonWriter, Property, (double)Cast<UIntProperty>(Property)->GetPropertyValue_InContainer(Data, ArrayIndex));
+		WritePropertyValue(JsonWriter, State.ValueProperty, (double)Cast<UIntProperty>(State.ValueProperty)->GetPropertyValue_InContainer(State.ValueData, ArrayIndex));
 	}
-	else if (TypeInfo == UInt8Property::StaticClass())
+	else if (State.ValueType == UInt8Property::StaticClass())
 	{
-		WritePropertyValue(JsonWriter, Property, (double)Cast<UInt8Property>(Property)->GetPropertyValue_InContainer(Data, ArrayIndex));
+		WritePropertyValue(JsonWriter, State.ValueProperty, (double)Cast<UInt8Property>(State.ValueProperty)->GetPropertyValue_InContainer(State.ValueData, ArrayIndex));
 	}
-	else if (TypeInfo == UInt16Property::StaticClass())
+	else if (State.ValueType == UInt16Property::StaticClass())
 	{
-		WritePropertyValue(JsonWriter, Property, (double)Cast<UInt16Property>(Property)->GetPropertyValue_InContainer(Data, ArrayIndex));
+		WritePropertyValue(JsonWriter, State.ValueProperty, (double)Cast<UInt16Property>(State.ValueProperty)->GetPropertyValue_InContainer(State.ValueData, ArrayIndex));
 	}
-	else if (TypeInfo == UInt64Property::StaticClass())
+	else if (State.ValueType == UInt64Property::StaticClass())
 	{
-		WritePropertyValue(JsonWriter, Property, (double)Cast<UInt64Property>(Property)->GetPropertyValue_InContainer(Data, ArrayIndex));
+		WritePropertyValue(JsonWriter, State.ValueProperty, (double)Cast<UInt64Property>(State.ValueProperty)->GetPropertyValue_InContainer(State.ValueData, ArrayIndex));
 	}
 
 	// unsigned integers
-	else if (TypeInfo == UUInt16Property::StaticClass())
+	else if (State.ValueType == UUInt16Property::StaticClass())
 	{
-		WritePropertyValue(JsonWriter, Property, (double)Cast<UUInt16Property>(Property)->GetPropertyValue_InContainer(Data, ArrayIndex));
+		WritePropertyValue(JsonWriter, State.ValueProperty, (double)Cast<UUInt16Property>(State.ValueProperty)->GetPropertyValue_InContainer(State.ValueData, ArrayIndex));
 	}
-	else if (TypeInfo == UUInt32Property::StaticClass())
+	else if (State.ValueType == UUInt32Property::StaticClass())
 	{
-		WritePropertyValue(JsonWriter, Property, (double)Cast<UUInt32Property>(Property)->GetPropertyValue_InContainer(Data, ArrayIndex));
+		WritePropertyValue(JsonWriter, State.ValueProperty, (double)Cast<UUInt32Property>(State.ValueProperty)->GetPropertyValue_InContainer(State.ValueData, ArrayIndex));
 	}
-	else if (TypeInfo == UUInt64Property::StaticClass())
+	else if (State.ValueType == UUInt64Property::StaticClass())
 	{
-		WritePropertyValue(JsonWriter, Property, (double)Cast<UUInt64Property>(Property)->GetPropertyValue_InContainer(Data, ArrayIndex));
+		WritePropertyValue(JsonWriter, State.ValueProperty, (double)Cast<UUInt64Property>(State.ValueProperty)->GetPropertyValue_InContainer(State.ValueData, ArrayIndex));
 	}
 
 	// names, strings & text
-	else if (TypeInfo == UNameProperty::StaticClass())
+	else if (State.ValueType == UNameProperty::StaticClass())
 	{
-		WritePropertyValue(JsonWriter, Property, Cast<UNameProperty>(Property)->GetPropertyValue_InContainer(Data, ArrayIndex).ToString());
+		WritePropertyValue(JsonWriter, State.ValueProperty, Cast<UNameProperty>(State.ValueProperty)->GetPropertyValue_InContainer(State.ValueData, ArrayIndex).ToString());
 	}
-	else if (TypeInfo == UStrProperty::StaticClass())
+	else if (State.ValueType == UStrProperty::StaticClass())
 	{
-		WritePropertyValue(JsonWriter, Property, Cast<UStrProperty>(Property)->GetPropertyValue_InContainer(Data, ArrayIndex));
+		WritePropertyValue(JsonWriter, State.ValueProperty, Cast<UStrProperty>(State.ValueProperty)->GetPropertyValue_InContainer(State.ValueData, ArrayIndex));
 	}
-	else if (TypeInfo == UTextProperty::StaticClass())
+	else if (State.ValueType == UTextProperty::StaticClass())
 	{
-		WritePropertyValue(JsonWriter, Property, Cast<UTextProperty>(Property)->GetPropertyValue_InContainer(Data, ArrayIndex).ToString());
+		WritePropertyValue(JsonWriter, State.ValueProperty, Cast<UTextProperty>(State.ValueProperty)->GetPropertyValue_InContainer(State.ValueData, ArrayIndex).ToString());
 	}
 
 	// classes & objects
-	else if (TypeInfo == UClassProperty::StaticClass())
+	else if (State.ValueType == UClassProperty::StaticClass())
 	{
-		WritePropertyValue(JsonWriter, Property, Cast<UClassProperty>(Property)->GetPropertyValue_InContainer(Data, ArrayIndex)->GetPathName());
+		WritePropertyValue(JsonWriter, State.ValueProperty, Cast<UClassProperty>(State.ValueProperty)->GetPropertyValue_InContainer(State.ValueData, ArrayIndex)->GetPathName());
 	}
-	else if (TypeInfo == UObjectProperty::StaticClass())
+	else if (State.ValueType == UObjectProperty::StaticClass())
 	{
-		WriteNull(JsonWriter, Property);
+		WriteNull(JsonWriter, State.ValueProperty);
 	}
 	
+	// unsupported property type
 	else
 	{
-		UE_LOG(LogSerialization, Verbose, TEXT("FJsonStructSerializerBackend: Property %s cannot be serialized, because its type (%s) is not supported"), *Property->GetFName().ToString(), *TypeInfo->GetFName().ToString());
+		UE_LOG(LogSerialization, Verbose, TEXT("FJsonStructSerializerBackend: Property %s cannot be serialized, because its type (%s) is not supported"), *State.ValueProperty->GetFName().ToString(), *State.ValueType->GetFName().ToString());
 	}
 }
diff --git a/Engine/Source/Runtime/Serialization/Private/StructSerializer.cpp b/Engine/Source/Runtime/Serialization/Private/StructSerializer.cpp
--- a/Engine/Source/Runtime/Serialization/Private/StructSerializer.cpp
+++ b/Engine/Source/Runtime/Serialization/Private/StructSerializer.cpp
@@ -10,25 +10,6 @@
 
 namespace StructSerializer
 {
-	/**
-	 * Structure for the write state stack.
-	 */
-	struct FWriteState
-	{
-		/** Holds a pointer to the property's data. */
-		const void* Data;
-
-		/** Holds a flag indicating whether the property has been processed. */
-		bool HasBeenProcessed;
-
-		/** Holds the property's meta data. */
-		UProperty* Property;
-
-		/** Holds a pointer to the UStruct describing the data. */
-		UStruct* TypeInfo;
-	};
-
-
 	/**
 	 * Gets the value from the given property.
 	 *
@@ -38,16 +19,16 @@ namespace StructSerializer
 	 * @return A pointer to the property's value, or nullptr if it couldn't be found.
 	 */
 	template<typename UPropertyType, typename PropertyType>
-	PropertyType* GetPropertyValue( const FWriteState& State, UProperty* Property )
+	PropertyType* GetPropertyValue( const FStructSerializerState& State, UProperty* Property )
 	{
 		PropertyType* ValuePtr = nullptr;
-		UArrayProperty* ArrayProperty = Cast<UArrayProperty>(State.Property);
+		UArrayProperty* ArrayProperty = Cast<UArrayProperty>(State.ValueProperty);
 
 		if (ArrayProperty)
 		{
 			check(ArrayProperty->Inner == Property);
 
-			FScriptArrayHelper ArrayHelper(ArrayProperty, ArrayProperty->template ContainerPtrToValuePtr<void>(State.Data));
+			FScriptArrayHelper ArrayHelper(ArrayProperty, ArrayProperty->template ContainerPtrToValuePtr<void>(State.ValueData));
 			int32 Index = ArrayHelper.AddValue();
 		
 			ValuePtr = (PropertyType*)ArrayHelper.GetRawPtr( Index );
@@ -57,7 +38,7 @@ namespace StructSerializer
 			UPropertyType* TypedProperty = Cast<UPropertyType>(Property);
 			check(TypedProperty != nullptr);
 
-			ValuePtr = TypedProperty->template ContainerPtrToValuePtr<PropertyType>(State.Data);
+			ValuePtr = TypedProperty->template ContainerPtrToValuePtr<PropertyType>(State.ValueData);
 		}
 
 		return ValuePtr;
@@ -75,14 +56,15 @@ void FStructSerializer::Serialize( const void* Struct, UStruct& TypeInfo, IStruc
 	check(Struct != nullptr);
 
 	// initialize serialization
-	TArray<FWriteState> StateStack;
+	TArray<FStructSerializerState> StateStack;
 	{
-		FWriteState NewState;
-
-		NewState.Data = Struct;
-		NewState.Property = nullptr;
-		NewState.TypeInfo = &TypeInfo;
-		NewState.HasBeenProcessed = false;
+		FStructSerializerState NewState;
+		{
+			NewState.HasBeenProcessed = false;
+			NewState.ValueData = Struct;
+			NewState.ValueProperty = nullptr;
+			NewState.ValueType = &TypeInfo;
+		}
 
 		StateStack.Push(NewState);
 	}
@@ -90,70 +72,68 @@ void FStructSerializer::Serialize( const void* Struct, UStruct& TypeInfo, IStruc
 	// process state stack
 	while (StateStack.Num() > 0)
 	{
-		FWriteState CurrentState = StateStack.Pop(/*bAllowShrinking=*/ false);
+		FStructSerializerState CurrentState = StateStack.Pop(/*bAllowShrinking=*/ false);
 
 		// structures
-		if ((CurrentState.Property == nullptr) || (CurrentState.TypeInfo == UStructProperty::StaticClass()))
+		if ((CurrentState.ValueProperty == nullptr) || (CurrentState.ValueType == UStructProperty::StaticClass()))
 		{
 			if (!CurrentState.HasBeenProcessed)
 			{
-				const void* NewData = CurrentState.Data;
+				const void* NewData = CurrentState.ValueData;
 
 				// write object start
-				if (CurrentState.Property == nullptr)
-				{
-					Backend.BeginStructure(CurrentState.TypeInfo);
-				}
-				else
+				if (CurrentState.ValueProperty != nullptr)
 				{
-					UObject* Outer = CurrentState.Property->GetOuter();
+					UObject* Outer = CurrentState.ValueProperty->GetOuter();
 
 					if ((Outer == nullptr) || (Outer->GetClass() != UArrayProperty::StaticClass()))
 					{
-						NewData = CurrentState.Property->ContainerPtrToValuePtr<void>(CurrentState.Data);
+						NewData = CurrentState.ValueProperty->ContainerPtrToValuePtr<void>(CurrentState.ValueData);
 					}
-
-					Backend.BeginStructure(CurrentState.Property);
 				}
 
+				Backend.BeginStructure(CurrentState);
+
 				CurrentState.HasBeenProcessed = true;
 				StateStack.Push(CurrentState);
 
 				// serialize fields
-				if (CurrentState.Property != nullptr)
+				if (CurrentState.ValueProperty != nullptr)
 				{
-					UStructProperty* StructProperty = Cast<UStructProperty>(CurrentState.Property);
+					UStructProperty* StructProperty = Cast<UStructProperty>(CurrentState.ValueProperty);
 
 					if (StructProperty != nullptr)
 					{
-						CurrentState.TypeInfo = StructProperty->Struct;
+						CurrentState.ValueType = StructProperty->Struct;
 					}
 					else
 					{
-						UObjectPropertyBase* ObjectProperty = Cast<UObjectPropertyBase>(CurrentState.Property);
+						UObjectPropertyBase* ObjectProperty = Cast<UObjectPropertyBase>(CurrentState.ValueProperty);
 
 						if (ObjectProperty != nullptr)
 						{
-							CurrentState.TypeInfo = ObjectProperty->PropertyClass;
+							CurrentState.ValueType = ObjectProperty->PropertyClass;
 						}
 					}
 				}
 
-				TArray<FWriteState> NewStates;
+				TArray<FStructSerializerState> NewStates;
 
-				for (TFieldIterator<UProperty> It(CurrentState.TypeInfo, EFieldIteratorFlags::IncludeSuper); It; ++It)
+				for (TFieldIterator<UProperty> It(CurrentState.ValueType, EFieldIteratorFlags::IncludeSuper); It; ++It)
 				{
 					// Skip property if the filter function is set and rejects it.
-					if (Policies.PropertyFilter && !Policies.PropertyFilter(*It, CurrentState.Property))
+					if (Policies.PropertyFilter && !Policies.PropertyFilter(*It, CurrentState.ValueProperty))
 					{
 						continue;
 					}
-					FWriteState NewState;
 
-					NewState.Data = NewData;
-					NewState.Property = *It;
-					NewState.TypeInfo = It->GetClass();
-					NewState.HasBeenProcessed = false;
+					FStructSerializerState NewState;
+					{
+						NewState.HasBeenProcessed = false;
+						NewState.ValueData = NewData;
+						NewState.ValueProperty = *It;
+						NewState.ValueType = It->GetClass();
+					}
 
 					NewStates.Add(NewState);
 				}
@@ -166,60 +146,61 @@ void FStructSerializer::Serialize( const void* Struct, UStruct& TypeInfo, IStruc
 			}
 			else
 			{
-				Backend.EndStructure();
+				Backend.EndStructure(CurrentState);
 			}
 		}
 
 		// dynamic arrays
-		else if (CurrentState.TypeInfo == UArrayProperty::StaticClass())
+		else if (CurrentState.ValueType == UArrayProperty::StaticClass())
 		{
 			if (!CurrentState.HasBeenProcessed)
 			{
-				Backend.BeginArray(CurrentState.Property);
+				Backend.BeginArray(CurrentState);
 
 				CurrentState.HasBeenProcessed = true;
 				StateStack.Push(CurrentState);
 
-				UArrayProperty* ArrayProperty = Cast<UArrayProperty>(CurrentState.Property);
-				FScriptArrayHelper ArrayHelper(ArrayProperty, ArrayProperty->ContainerPtrToValuePtr<void>(CurrentState.Data));
+				UArrayProperty* ArrayProperty = Cast<UArrayProperty>(CurrentState.ValueProperty);
+				FScriptArrayHelper ArrayHelper(ArrayProperty, ArrayProperty->ContainerPtrToValuePtr<void>(CurrentState.ValueData));
 				UProperty* Inner = ArrayProperty->Inner;
 
 				// push elements on stack (in reverse order)
 				for (int Index = ArrayHelper.Num() - 1; Index >= 0; --Index)
 				{
-					FWriteState NewState;
-
-					NewState.Data = ArrayHelper.GetRawPtr(Index);
-					NewState.Property = Inner;
-					NewState.TypeInfo = Inner->GetClass();
-					NewState.HasBeenProcessed = false;
+					FStructSerializerState NewState;
+					{
+						NewState.HasBeenProcessed = false;
+						NewState.ValueData = ArrayHelper.GetRawPtr(Index);
+						NewState.ValueProperty = Inner;
+						NewState.ValueType = Inner->GetClass();
+					}
 
 					StateStack.Push(NewState);
 				}
 			}
 			else
 			{
-				Backend.EndArray(CurrentState.Property);
+				Backend.EndArray(CurrentState);
 			}
 		}
 
 		// static arrays
-		else if (CurrentState.Property->ArrayDim > 1)
+		else if (CurrentState.ValueProperty->ArrayDim > 1)
 		{
-			Backend.BeginArray(CurrentState.Property);
+			Backend.BeginArray(CurrentState);
 
-			for (int32 ArrayIndex = 0; ArrayIndex < CurrentState.Property->ArrayDim; ++ArrayIndex)
+			for (int32 ArrayIndex = 0; ArrayIndex < CurrentState.ValueProperty->ArrayDim; ++ArrayIndex)
 			{
-				Backend.WriteProperty(CurrentState.Property, CurrentState.Data, CurrentState.TypeInfo, ArrayIndex);
+				Backend.WriteProperty(CurrentState, ArrayIndex);
 			}
 
-			Backend.EndArray(CurrentState.Property);
+			Backend.EndArray(CurrentState);
 		}
 
 		// all other properties
 		else
 		{
-			Backend.WriteProperty(CurrentState.Property, CurrentState.Data, CurrentState.TypeInfo);
+			Backend.WriteProperty(CurrentState);
 		}
 	}
 }
diff --git a/Engine/Source/Runtime/Serialization/Public/Backends/JsonStructSerializerBackend.h b/Engine/Source/Runtime/Serialization/Public/Backends/JsonStructSerializerBackend.h
--- a/Engine/Source/Runtime/Serialization/Public/Backends/JsonStructSerializerBackend.h
+++ b/Engine/Source/Runtime/Serialization/Public/Backends/JsonStructSerializerBackend.h
@@ -37,13 +37,12 @@ public:
 
 	// IStructSerializerBackend interface
 
-	virtual void BeginArray( UProperty* Property ) override;
-	virtual void BeginStructure( UProperty* Property ) override;
-	virtual void BeginStructure( UStruct* TypeInfo ) override;
-	virtual void EndArray( UProperty* Property ) override;
-	virtual void EndStructure() override;
-	virtual void WriteComment( const FString& Comment ) override;
-	virtual void WriteProperty( UProperty* Property, const void* Data, UStruct* TypeInfo, int32 ArrayIndex ) override;
+	virtual void BeginArray(const FStructSerializerState& State) override;
+	virtual void BeginStructure(const FStructSerializerState& State) override;
+	virtual void EndArray(const FStructSerializerState& State) override;
+	virtual void EndStructure(const FStructSerializerState& State) override;
+	virtual void WriteComment(const FString& Comment) override;
+	virtual void WriteProperty(const FStructSerializerState& State, int32 ArrayIndex = 0) override;
 
 private:
 
diff --git a/Engine/Source/Runtime/Serialization/Public/IStructSerializerBackend.h b/Engine/Source/Runtime/Serialization/Public/IStructSerializerBackend.h
--- a/Engine/Source/Runtime/Serialization/Public/IStructSerializerBackend.h
+++ b/Engine/Source/Runtime/Serialization/Public/IStructSerializerBackend.h
@@ -8,6 +8,25 @@ class UProperty;
 class UStruct;
 
 
+/**
+ * Structure for the write state stack.
+ */
+struct FStructSerializerState
+{
+	/** Holds a flag indicating whether the property has been processed. */
+	bool HasBeenProcessed;
+
+	/** Holds a pointer to the property value's data. */
+	const void* ValueData;
+
+	/** Holds the property value's meta data. */
+	UProperty* ValueProperty;
+
+	/** Holds a pointer to the UStruct describing the data. */
+	UStruct* ValueType;
+};
+
+
 /**
  * Interface for UStruct serializer backends.
  */
@@ -18,54 +37,65 @@ public:
 	/**
 	 * Signals the beginning of an array.
 	 *
-	 * @param Property The property that holds the array.
+	 * State.ValueProperty points to the property that holds the array.
+	 *
+	 * @param State The serializer's current state.
+	 * @see BeginStructure, EndArray
 	 */
-	virtual void BeginArray( UProperty* Property ) = 0;
+	virtual void BeginArray(const FStructSerializerState& State) = 0;
 
 	/**
 	 * Signals the beginning of a child structure.
 	 *
-	 * @param Property The property that holds the object.
-	 */
-	virtual void BeginStructure( UProperty* Property ) = 0;
-
-	/**
-	 * Signals the beginning of a root structure.
+	 * State.ValueProperty points to the property that holds the struct.
 	 *
-	 * @param TypeInfo The object's type information.
+	 * @param State The serializer's current state.
+	 * @see BeginArray, EndStructure
 	 */
-	virtual void BeginStructure( UStruct* TypeInfo ) = 0;
+	virtual void BeginStructure(const FStructSerializerState& State) = 0;
 
 	/**
 	 * Signals the end of an array.
 	 *
-	 * @param Property The property that holds the array.
+	 * State.ValueProperty points to the property that holds the array.
+	 *
+	 * @param State The serializer's current state.
+	 * @see BeginArray, EndStructure
 	 */
-	virtual void EndArray( UProperty* Property ) = 0;
+	virtual void EndArray(const FStructSerializerState& State) = 0;
 
 	/**
 	 * Signals the end of an object.
+	 *
+	 * State.ValueProperty points to the property that holds the struct.
+	 *
+	 * @param State The serializer's current state.
+	 * @see BeginStructure, EndArray
 	 */
-	virtual void EndStructure() = 0;
+	virtual void EndStructure(const FStructSerializerState& State) = 0;
 
 	/**
 	 * Writes a comment to the output stream.
 	 *
 	 * @param Comment The comment text.
+	 * @see BeginArray, BeginStructure, EndArray, EndStructure, WriteProperty
 	 */
-	virtual void WriteComment( const FString& Comment ) = 0;
+	virtual void WriteComment(const FString& Comment) = 0;
 
 	/**
 	 * Writes a property to the output stream.
 	 *
 	 * Depending on the context, properties to be written can be either object properties or array elements.
 	 *
-	 * @param Property The property that holds the data to write.
-	 * @param Data The property's data to write.
-	 * @param TypeInfo The property data's type information.
-	 * @param ArrayIndex An index into the property array (for static arrays).
+	 * State.ValueProperty points to the property that holds the value to write.
+	 * State.ValueData points to the actual data to write.
+	 * State.TypeInfo contains the data's type information
+	 * State.ArrayIndex is the optional index if the data is a value in an array.
+	 *
+	 * @param State The serializer's current state.
+	 * @see BeginArray, BeginStructure, EndArray, EndStructure, WriteComment
 	 */
-	virtual void WriteProperty( UProperty* Property, const void* Data, UStruct* TypeInfo, int32 ArrayIndex = 0 ) = 0;
+	virtual void WriteProperty(const FStructSerializerState& State, int32 ArrayIndex = 0) = 0;
 
 public:
 
diff --git a/Engine/Source/Runtime/WebBrowser/Private/WebJSStructSerializerBackend.cpp b/Engine/Source/Runtime/WebBrowser/Private/WebJSStructSerializerBackend.cpp
--- a/Engine/Source/Runtime/WebBrowser/Private/WebJSStructSerializerBackend.cpp
+++ b/Engine/Source/Runtime/WebBrowser/Private/WebJSStructSerializerBackend.cpp
@@ -89,26 +89,27 @@ void FWebJSStructSerializerBackend::Add(UProperty* Property, UObject* Value)
 /* IStructSerializerBackend interface
  *****************************************************************************/
 
-void FWebJSStructSerializerBackend::BeginArray( UProperty* Property )
+void FWebJSStructSerializerBackend::BeginArray(const FStructSerializerState& State)
 {
 	CefRefPtr<CefListValue> ListValue = CefListValue::Create();
-	Stack.Push(StackItem(Property->GetName(), ListValue));
+	Stack.Push(StackItem(State.ValueProperty->GetName(), ListValue));
 }
 
-void FWebJSStructSerializerBackend::BeginStructure( UProperty* Property )
+void FWebJSStructSerializerBackend::BeginStructure(const FStructSerializerState& State)
 {
-	CefRefPtr<CefDictionaryValue> DictionaryValue = CefDictionaryValue::Create();
-	Stack.Push(StackItem(Property->GetName(), DictionaryValue));
-}
-
-void FWebJSStructSerializerBackend::BeginStructure( UStruct* TypeInfo )
-{
-	Result = CefDictionaryValue::Create();
-	Stack.Push(StackItem(FString(), Result));
+	if (State.ValueProperty != nullptr)
+	{
+		CefRefPtr<CefDictionaryValue> DictionaryValue = CefDictionaryValue::Create();
+		Stack.Push(StackItem(State.ValueProperty->GetName(), DictionaryValue));
+	}
+	else
+	{
+		Result = CefDictionaryValue::Create();
+		Stack.Push(StackItem(FString(), Result));
+	}
 }
 
-
-void FWebJSStructSerializerBackend::EndArray( UProperty* Property )
+void FWebJSStructSerializerBackend::EndArray(const FStructSerializerState& /*State*/)
 {
 	StackItem Previous = Stack.Pop();
 	check(Previous.Kind == StackItem::STYPE_LIST);
@@ -126,7 +127,7 @@ void FWebJSStructSerializerBackend::EndArray( UProperty* Property )
 }
 
 
-void FWebJSStructSerializerBackend::EndStructure()
+void FWebJSStructSerializerBackend::EndStructure(const FStructSerializerState& /*State*/)
 {
 	StackItem Previous = Stack.Pop();
 	check(Previous.Kind == StackItem::STYPE_DICTIONARY);
@@ -151,105 +152,106 @@ void FWebJSStructSerializerBackend::EndStructure()
 }
 
 
-void FWebJSStructSerializerBackend::WriteComment( const FString& Comment )
+void FWebJSStructSerializerBackend::WriteComment(const FString& Comment)
 {
 	// Cef values do not support comments
 }
 
 
-void FWebJSStructSerializerBackend::WriteProperty( UProperty* Property, const void* Data, UStruct* TypeInfo, int32 ArrayIndex )
+void FWebJSStructSerializerBackend::WriteProperty(const FStructSerializerState& State, int32 ArrayIndex)
 {
 
 	// booleans
-	if (TypeInfo == UBoolProperty::StaticClass())
+	if (State.ValueType == UBoolProperty::StaticClass())
 	{
-		Add(Property, Cast<UBoolProperty>(Property)->GetPropertyValue_InContainer(Data, ArrayIndex));
+		Add(State.ValueProperty, Cast<UBoolProperty>(State.ValueProperty)->GetPropertyValue_InContainer(State.ValueData, ArrayIndex));
 	}
 
 	// unsigned bytes & enumerations
-	else if (TypeInfo == UByteProperty::StaticClass())
+	else if (State.ValueType == UByteProperty::StaticClass())
 	{
-		UByteProperty* ByteProperty = Cast<UByteProperty>(Property);
+		UByteProperty* ByteProperty = Cast<UByteProperty>(State.ValueProperty);
 
 		if (ByteProperty->IsEnum())
 		{
-			Add(Property, ByteProperty->Enum->GetEnumName(ByteProperty->GetPropertyValue_InContainer(Data, ArrayIndex)));
+			Add(State.ValueProperty, ByteProperty->Enum->GetEnumName(ByteProperty->GetPropertyValue_InContainer(State.ValueData, ArrayIndex)));
 		}
 		else
 		{
-			Add(Property, (double)Cast<UByteProperty>(Property)->GetPropertyValue_InContainer(Data, ArrayIndex));
+			Add(State.ValueProperty, (double)Cast<UByteProperty>(State.ValueProperty)->GetPropertyValue_InContainer(State.ValueData, ArrayIndex));
 		}			
 	}
 
 	// floating point numbers
-	else if (TypeInfo == UDoubleProperty::StaticClass())
+	else if (State.ValueType == UDoubleProperty::StaticClass())
 	{
-		Add(Property, Cast<UDoubleProperty>(Property)->GetPropertyValue_InContainer(Data, ArrayIndex));
+		Add(State.ValueProperty, Cast<UDoubleProperty>(State.ValueProperty)->GetPropertyValue_InContainer(State.ValueData, ArrayIndex));
 	}
-	else if (TypeInfo == UFloatProperty::StaticClass())
+	else if (State.ValueType == UFloatProperty::StaticClass())
 	{
-		Add(Property, Cast<UFloatProperty>(Property)->GetPropertyValue_InContainer(Data, ArrayIndex));
+		Add(State.ValueProperty, Cast<UFloatProperty>(State.ValueProperty)->GetPropertyValue_InContainer(State.ValueData, ArrayIndex));
 	}
 
 	// signed integers
-	else if (TypeInfo == UIntProperty::StaticClass())
+	else if (State.ValueType == UIntProperty::StaticClass())
 	{
-		Add(Property, (int32)Cast<UIntProperty>(Property)->GetPropertyValue_InContainer(Data, ArrayIndex));
+		Add(State.ValueProperty, (int32)Cast<UIntProperty>(State.ValueProperty)->GetPropertyValue_InContainer(State.ValueData, ArrayIndex));
 	}
-	else if (TypeInfo == UInt8Property::StaticClass())
+	else if (State.ValueType == UInt8Property::StaticClass())
 	{
-		Add(Property, (int32)Cast<UInt8Property>(Property)->GetPropertyValue_InContainer(Data, ArrayIndex));
+		Add(State.ValueProperty, (int32)Cast<UInt8Property>(State.ValueProperty)->GetPropertyValue_InContainer(State.ValueData, ArrayIndex));
 	}
-	else if (TypeInfo == UInt16Property::StaticClass())
+	else if (State.ValueType == UInt16Property::StaticClass())
 	{
-		Add(Property, (int32)Cast<UInt16Property>(Property)->GetPropertyValue_InContainer(Data, ArrayIndex));
+		Add(State.ValueProperty, (int32)Cast<UInt16Property>(State.ValueProperty)->GetPropertyValue_InContainer(State.ValueData, ArrayIndex));
 	}
-	else if (TypeInfo == UInt64Property::StaticClass())
+	else if (State.ValueType == UInt64Property::StaticClass())
 	{
-		Add(Property, (double)Cast<UInt64Property>(Property)->GetPropertyValue_InContainer(Data, ArrayIndex));
+		Add(State.ValueProperty, (double)Cast<UInt64Property>(State.ValueProperty)->GetPropertyValue_InContainer(State.ValueData, ArrayIndex));
 	}
 
 	// unsigned integers
-	else if (TypeInfo == UUInt16Property::StaticClass())
+	else if (State.ValueType == UUInt16Property::StaticClass())
 	{
-		Add(Property, (int32)Cast<UUInt16Property>(Property)->GetPropertyValue_InContainer(Data, ArrayIndex));
+		Add(State.ValueProperty, (int32)Cast<UUInt16Property>(State.ValueProperty)->GetPropertyValue_InContainer(State.ValueData, ArrayIndex));
 	}
-	else if (TypeInfo == UUInt32Property::StaticClass())
+	else if (State.ValueType == UUInt32Property::StaticClass())
 	{
-		Add(Property, (double)Cast<UUInt32Property>(Property)->GetPropertyValue_InContainer(Data, ArrayIndex));
+		Add(State.ValueProperty, (double)Cast<UUInt32Property>(State.ValueProperty)->GetPropertyValue_InContainer(State.ValueData, ArrayIndex));
 	}
-	else if (TypeInfo == UUInt64Property::StaticClass())
+	else if (State.ValueType == UUInt64Property::StaticClass())
 	{
-		Add(Property, (double)Cast<UUInt64Property>(Property)->GetPropertyValue_InContainer(Data, ArrayIndex));
+		Add(State.ValueProperty, (double)Cast<UUInt64Property>(State.ValueProperty)->GetPropertyValue_InContainer(State.ValueData, ArrayIndex));
 	}
 
 	// names & strings
-	else if (TypeInfo == UNameProperty::StaticClass())
+	else if (State.ValueType == UNameProperty::StaticClass())
 	{
-		Add(Property, Cast<UNameProperty>(Property)->GetPropertyValue_InContainer(Data, ArrayIndex).ToString());
+		Add(State.ValueProperty, Cast<UNameProperty>(State.ValueProperty)->GetPropertyValue_InContainer(State.ValueData, ArrayIndex).ToString());
 	}
-	else if (TypeInfo == UStrProperty::StaticClass())
+	else if (State.ValueType == UStrProperty::StaticClass())
 	{
-		Add(Property, Cast<UStrProperty>(Property)->GetPropertyValue_InContainer(Data, ArrayIndex));
+		Add(State.ValueProperty, Cast<UStrProperty>(State.ValueProperty)->GetPropertyValue_InContainer(State.ValueData, ArrayIndex));
 	}
-	else if (TypeInfo == UTextProperty::StaticClass())
+	else if (State.ValueType == UTextProperty::StaticClass())
 	{
-		Add(Property, Cast<UTextProperty>(Property)->GetPropertyValue_InContainer(Data, ArrayIndex).ToString());
+		Add(State.ValueProperty, Cast<UTextProperty>(State.ValueProperty)->GetPropertyValue_InContainer(State.ValueData, ArrayIndex).ToString());
 	}
 
 	// classes & objects
-	else if (TypeInfo == UClassProperty::StaticClass())
+	else if (State.ValueType == UClassProperty::StaticClass())
 	{
-		Add(Property, Cast<UClassProperty>(Property)->GetPropertyValue_InContainer(Data, ArrayIndex)->GetPathName());
+		Add(State.ValueProperty, Cast<UClassProperty>(State.ValueProperty)->GetPropertyValue_InContainer(State.ValueData, ArrayIndex)->GetPathName());
 	}
-	else if (TypeInfo == UObjectProperty::StaticClass())
+	else if (State.ValueType == UObjectProperty::StaticClass())
 	{
-		Add(Property, Cast<UObjectProperty>(Property)->GetPropertyValue_InContainer(Data, ArrayIndex));
+		Add(State.ValueProperty, Cast<UObjectProperty>(State.ValueProperty)->GetPropertyValue_InContainer(State.ValueData, ArrayIndex));
 	}
 	
+	// unsupported property type
 	else
 	{
-		GLog->Logf(ELogVerbosity::Warning, TEXT("FWebJSStructSerializerBackend: Property %s cannot be serialized, because its type (%s) is not supported"), *Property->GetName(), *TypeInfo->GetName());
+		GLog->Logf(ELogVerbosity::Warning, TEXT("FWebJSStructSerializerBackend: Property %s cannot be serialized, because its type (%s) is not supported"), *State.ValueProperty->GetName(), *State.ValueType->GetName());
 	}
 }
 
diff --git a/Engine/Source/Runtime/WebBrowser/Private/WebJSStructSerializerBackend.h b/Engine/Source/Runtime/WebBrowser/Private/WebJSStructSerializerBackend.h
--- a/Engine/Source/Runtime/WebBrowser/Private/WebJSStructSerializerBackend.h
+++ b/Engine/Source/Runtime/WebBrowser/Private/WebJSStructSerializerBackend.h
@@ -47,13 +47,12 @@ public:
 
 	// IStructSerializerBackend interface
 
-	virtual void BeginArray( UProperty* Property ) override;
-	virtual void BeginStructure( UProperty* Property ) override;
-	virtual void BeginStructure( UStruct* TypeInfo ) override;
-	virtual void EndArray( UProperty* Property ) override;
-	virtual void EndStructure() override;
-	virtual void WriteComment( const FString& Comment ) override;
-	virtual void WriteProperty( UProperty* Property, const void* Data, UStruct* TypeInfo, int32 ArrayIndex ) override;
+	virtual void BeginArray(const FStructSerializerState& State) override;
+	virtual void BeginStructure(const FStructSerializerState& State) override;
+	virtual void EndArray(const FStructSerializerState& State) override;
+	virtual void EndStructure(const FStructSerializerState& State) override;
+	virtual void WriteComment(const FString& Comment) override;
+	virtual void WriteProperty(const FStructSerializerState& State, int32 ArrayIndex = 0) override;
 
 private:
 
