Serialization: Treating dictionaries as structs (for symmetry with, and simplified de-serialization)

diff --git a/Engine/Source/Runtime/Serialization/Private/Backends/JsonStructSerializerBackend.cpp b/Engine/Source/Runtime/Serialization/Private/Backends/JsonStructSerializerBackend.cpp
--- a/Engine/Source/Runtime/Serialization/Private/Backends/JsonStructSerializerBackend.cpp
+++ b/Engine/Source/Runtime/Serialization/Private/Backends/JsonStructSerializerBackend.cpp
@@ -74,27 +74,6 @@ void FJsonStructSerializerBackend::BeginArray(const FStructSerializerState& Stat
 }
 
 
-void FJsonStructSerializerBackend::BeginDictionary(const FStructSerializerState& State)
-{
-	UObject* Outer = State.ValueProperty->GetOuter();
-
-	if ((Outer != nullptr) && (Outer->GetClass() == UArrayProperty::StaticClass()))
-	{
-		JsonWriter->WriteObjectStart();
-	}
-	else if (State.KeyProperty != nullptr)
-	{
-		FString KeyString;
-		State.KeyProperty->ExportTextItem(KeyString, State.KeyData, nullptr, nullptr, PPF_None);
-		JsonWriter->WriteObjectStart(KeyString);
-	}
-	else
-	{
-		JsonWriter->WriteObjectStart(State.ValueProperty->GetName());
-	}
-}
-
-
 void FJsonStructSerializerBackend::BeginStructure(const FStructSerializerState& State)
 {
 	if (State.ValueProperty != nullptr)
@@ -129,12 +108,6 @@ void FJsonStructSerializerBackend::EndArray(const FStructSerializerState& /*Stat
 }
 
 
-void FJsonStructSerializerBackend::EndDictionary(const FStructSerializerState& /*State*/)
-{
-	JsonWriter->WriteObjectEnd();
-}
-
-
 void FJsonStructSerializerBackend::EndStructure(const FStructSerializerState& /*State*/)
 {
 	JsonWriter->WriteObjectEnd();
diff --git a/Engine/Source/Runtime/Serialization/Private/StructSerializer.cpp b/Engine/Source/Runtime/Serialization/Private/StructSerializer.cpp
--- a/Engine/Source/Runtime/Serialization/Private/StructSerializer.cpp
+++ b/Engine/Source/Runtime/Serialization/Private/StructSerializer.cpp
@@ -195,7 +195,7 @@ void FStructSerializer::Serialize( const void* Struct, UStruct& TypeInfo, IStruc
 		{
 			if (!CurrentState.HasBeenProcessed)
 			{
-				Backend.BeginDictionary(CurrentState);
+				Backend.BeginStructure(CurrentState);
 
 				CurrentState.HasBeenProcessed = true;
 				StateStack.Push(CurrentState);
@@ -214,7 +214,7 @@ void FStructSerializer::Serialize( const void* Struct, UStruct& TypeInfo, IStruc
 						NewState.HasBeenProcessed = false;
 						NewState.KeyData = PairPtr + MapProperty->MapLayout.KeyOffset;
 						NewState.KeyProperty = MapProperty->KeyProp;
-						NewState.ValueData = PairPtr;// + MapProperty->MapLayout.ValueOffset;
+						NewState.ValueData = PairPtr;
 						NewState.ValueProperty = ValueProperty;
 						NewState.ValueType = ValueProperty->GetClass();
 					}
@@ -224,7 +224,7 @@ void FStructSerializer::Serialize( const void* Struct, UStruct& TypeInfo, IStruc
 			}
 			else
 			{
-				Backend.EndDictionary(CurrentState);
+				Backend.EndStructure(CurrentState);
 			}
 		}
 
diff --git a/Engine/Source/Runtime/Serialization/Public/Backends/JsonStructSerializerBackend.h b/Engine/Source/Runtime/Serialization/Public/Backends/JsonStructSerializerBackend.h
--- a/Engine/Source/Runtime/Serialization/Public/Backends/JsonStructSerializerBackend.h
+++ b/Engine/Source/Runtime/Serialization/Public/Backends/JsonStructSerializerBackend.h
@@ -38,10 +38,8 @@ public:
 	// IStructSerializerBackend interface
 
 	virtual void BeginArray(const FStructSerializerState& State) override;
-	virtual void BeginDictionary(const FStructSerializerState& State) override;
 	virtual void BeginStructure(const FStructSerializerState& State) override;
 	virtual void EndArray(const FStructSerializerState& State) override;
-	virtual void EndDictionary(const FStructSerializerState& State) override;
 	virtual void EndStructure(const FStructSerializerState& State) override;
 	virtual void WriteComment(const FString& Comment) override;
 	virtual void WriteProperty(const FStructSerializerState& State, int32 ArrayIndex = 0) override;
diff --git a/Engine/Source/Runtime/Serialization/Public/IStructSerializerBackend.h b/Engine/Source/Runtime/Serialization/Public/IStructSerializerBackend.h
--- a/Engine/Source/Runtime/Serialization/Public/IStructSerializerBackend.h
+++ b/Engine/Source/Runtime/Serialization/Public/IStructSerializerBackend.h
@@ -50,13 +50,6 @@ public:
 	 */
 	virtual void BeginArray(const FStructSerializerState& State) = 0;
 
-	/**
-	 * Signals the beginning of a dictionary.
-	 *
-	 * @param Property The property that holds the dictionary.
-	 */
-	virtual void BeginDictionary(const FStructSerializerState& State) = 0;
-
 	/**
 	 * Signals the beginning of a child structure.
 	 *
@@ -77,13 +70,6 @@ public:
 	 */
 	virtual void EndArray(const FStructSerializerState& State) = 0;
 
-	/**
-	 * Signals the end of a dictionary.
-	 *
-	 * @param Property The property that holds the dictionary.
-	 */
-	virtual void EndDictionary(const FStructSerializerState& State) = 0;
-
 	/**
 	 * Signals the end of an object.
 	 *
diff --git a/Engine/Source/Runtime/WebBrowser/Private/WebJSStructSerializerBackend.cpp b/Engine/Source/Runtime/WebBrowser/Private/WebJSStructSerializerBackend.cpp
--- a/Engine/Source/Runtime/WebBrowser/Private/WebJSStructSerializerBackend.cpp
+++ b/Engine/Source/Runtime/WebBrowser/Private/WebJSStructSerializerBackend.cpp
@@ -103,13 +103,6 @@ void FWebJSStructSerializerBackend::BeginArray(const FStructSerializerState& Sta
 }
 
 
-void FWebJSStructSerializerBackend::BeginDictionary(const FStructSerializerState& State)
-{
-	CefRefPtr<CefDictionaryValue> DictionaryValue = CefDictionaryValue::Create();
-	Stack.Push(StackItem(State.ValueProperty->GetName(), DictionaryValue));
-}
-
-
 void FWebJSStructSerializerBackend::BeginStructure(const FStructSerializerState& State)
 {
 	if (State.KeyProperty != nullptr)
@@ -151,12 +144,6 @@ void FWebJSStructSerializerBackend::EndArray(const FStructSerializerState& /*Sta
 }
 
 
-void FWebJSStructSerializerBackend::EndDictionary(const FStructSerializerState& State)
-{
-	EndStructure(State);
-}
-
-
 void FWebJSStructSerializerBackend::EndStructure(const FStructSerializerState& /*State*/)
 {
 	StackItem Previous = Stack.Pop();
diff --git a/Engine/Source/Runtime/WebBrowser/Private/WebJSStructSerializerBackend.h b/Engine/Source/Runtime/WebBrowser/Private/WebJSStructSerializerBackend.h
--- a/Engine/Source/Runtime/WebBrowser/Private/WebJSStructSerializerBackend.h
+++ b/Engine/Source/Runtime/WebBrowser/Private/WebJSStructSerializerBackend.h
@@ -52,10 +52,8 @@ public:
 	// IStructSerializerBackend interface
 
 	virtual void BeginArray(const FStructSerializerState& State) override;
-	virtual void BeginDictionary(const FStructSerializerState& State) override;
 	virtual void BeginStructure(const FStructSerializerState& State) override;
 	virtual void EndArray(const FStructSerializerState& State) override;
-	virtual void EndDictionary(const FStructSerializerState& State) override;
 	virtual void EndStructure(const FStructSerializerState& State) override;
 	virtual void WriteComment(const FString& Comment) override;
 	virtual void WriteProperty(const FStructSerializerState& State, int32 ArrayIndex = 0) override;
