Fixed typo

diff --git a/Engine/Source/Runtime/Serialization/Private/StructSerializer.cpp b/Engine/Source/Runtime/Serialization/Private/StructSerializer.cpp
--- a/Engine/Source/Runtime/Serialization/Private/StructSerializer.cpp
+++ b/Engine/Source/Runtime/Serialization/Private/StructSerializer.cpp
@@ -207,7 +207,7 @@ void FStructSerializer::Serialize( const void* Struct, UStruct& TypeInfo, IStruc
 				// push key-value pairs on stack (in reverse order)
 				for (int Index = MapHelper.Num() - 1; Index >= 0; --Index)
 				{
-					uint8* PairPtr = MapHelper.GetPairPtr(Index);;
+					uint8* PairPtr = MapHelper.GetPairPtr(Index);
 
 					FStructSerializerState NewState;
 					{
