Serialization: Added support for TMap property de-serialization

diff --git a/Engine/Source/Runtime/Serialization/Private/StructDeserializer.cpp b/Engine/Source/Runtime/Serialization/Private/StructDeserializer.cpp
--- a/Engine/Source/Runtime/Serialization/Private/StructDeserializer.cpp
+++ b/Engine/Source/Runtime/Serialization/Private/StructDeserializer.cpp
@@ -109,7 +109,7 @@ bool FStructDeserializer::Deserialize( void* OutStruct, UStruct& TypeInfo, IStru
 					return false;
 				}
 
-				CurrentState = StateStack.Pop(/*bAllowShrinking=*/ false);
+				CurrentState = StateStack.Pop(/*bAllowShrinking*/ false);
 			}
 			break;
 
@@ -119,8 +119,24 @@ bool FStructDeserializer::Deserialize( void* OutStruct, UStruct& TypeInfo, IStru
 
 				NewState.Property = FindField<UProperty>(CurrentState.TypeInfo, *PropertyName);
 
-				if (NewState.Property == nullptr)
+				if (NewState.Property != nullptr)
+				{
+					// handle array property
+					if (Policies.PropertyFilter && Policies.PropertyFilter(NewState.Property, CurrentState.Property))
+					{
+						continue;
+					}
+
+					NewState.ArrayIndex = 0;
+					NewState.Data = CurrentState.Data;
+					NewState.TypeInfo = FindClass(NewState);
+
+					StateStack.Push(CurrentState);
+					CurrentState = NewState;
+				}
+				else
 				{
+					// error: array property not found
 					if (Policies.MissingFields != EStructDeserializerErrorPolicies::Ignore)
 					{
 						UE_LOG(LogSerialization, Verbose, TEXT("The array property '%s' does not exist"), *PropertyName);
@@ -133,20 +149,6 @@ bool FStructDeserializer::Deserialize( void* OutStruct, UStruct& TypeInfo, IStru
 
 					Backend.SkipArray();
 				}
-				else
-				{
-					// Skip array property if property filter is set and rejects it
-					if (Policies.PropertyFilter && Policies.PropertyFilter(NewState.Property, CurrentState.Property))
-					{
-						continue;
-					}
-					NewState.ArrayIndex = 0;
-					NewState.Data = CurrentState.Data;
-					NewState.TypeInfo = FindClass(NewState);
-
-					StateStack.Push(CurrentState);
-					CurrentState = NewState;
-				}
 			}
 			break;
 
@@ -157,27 +159,26 @@ bool FStructDeserializer::Deserialize( void* OutStruct, UStruct& TypeInfo, IStru
 
 		case EStructDeserializerBackendTokens::Property:
 			{
-				UProperty* Property = nullptr;
-
 				if (PropertyName.IsEmpty())
 				{
-					// handle array property
+					// handle array element
 					UArrayProperty* ArrayProperty = Cast<UArrayProperty>(CurrentState.Property);
+					UProperty* Property = nullptr;
 
 					if (ArrayProperty != nullptr)
 					{
-						// dynamic array property
+						// dynamic array element
 						Property = ArrayProperty->Inner;
 					}
 					else
 					{
-						// static array property
+						// static array element
 						Property = CurrentState.Property;
 					}
 
 					if (Property == nullptr)
 					{
-						// error: no property for array element
+						// error: no meta data for array element
 						if (Policies.MissingFields != EStructDeserializerErrorPolicies::Ignore)
 						{
 							UE_LOG(LogSerialization, Verbose, TEXT("Failed to serialize array element %i"), CurrentState.ArrayIndex);
@@ -192,35 +193,51 @@ bool FStructDeserializer::Deserialize( void* OutStruct, UStruct& TypeInfo, IStru
 
 					++CurrentState.ArrayIndex;
 				}
+				else if ((CurrentState.Property != nullptr) && (CurrentState.Property->GetClass() == UMapProperty::StaticClass()))
+				{
+					// handle map element
+					UMapProperty* MapProperty = Cast<UMapProperty>(CurrentState.Property);
+					FScriptMapHelper MapHelper(MapProperty, MapProperty->ContainerPtrToValuePtr<void>(CurrentState.Data));
+					UProperty* Property = MapProperty->ValueProp;
+
+					int32 PairIndex = MapHelper.AddDefaultValue_Invalid_NeedsRehash();
+					uint8* PairPtr = MapHelper.GetPairPtr(PairIndex);
+
+					MapProperty->KeyProp->ImportText(*PropertyName, PairPtr + MapProperty->MapLayout.KeyOffset, PPF_Delimited, nullptr);
+
+					if (!Backend.ReadProperty(Property, CurrentState.Property, PairPtr, CurrentState.ArrayIndex))
+					{
+						UE_LOG(LogSerialization, Verbose, TEXT("An item in map '%s' could not be read (%s)"), *PropertyName, *Backend.GetDebugString());
+					}
+				}
 				else
 				{
 					// handle scalar property
-					Property = FindField<UProperty>(CurrentState.TypeInfo, *PropertyName);
+					UProperty* Property = FindField<UProperty>(CurrentState.TypeInfo, *PropertyName);
 
-					if (Property == nullptr)
+					if (Property != nullptr)
 					{
-						// error: property not found
-						if (Policies.MissingFields != EStructDeserializerErrorPolicies::Ignore)
+						if (Policies.PropertyFilter && Policies.PropertyFilter(Property, CurrentState.Property))
 						{
-							UE_LOG(LogSerialization, Verbose, TEXT("The property '%s' does not exist"), *PropertyName);
+							continue;
 						}
 
-						if (Policies.MissingFields == EStructDeserializerErrorPolicies::Error)
+						if (!Backend.ReadProperty(Property, CurrentState.Property, CurrentState.Data, CurrentState.ArrayIndex))
 						{
-							return false;
+							UE_LOG(LogSerialization, Verbose, TEXT("The property '%s' could not be read (%s)"), *PropertyName, *Backend.GetDebugString());
 						}
 					}
 					else
 					{
-						// Skip scalar property if property filter is set and rejects it
-						if (Policies.PropertyFilter && Policies.PropertyFilter(Property, CurrentState.Property))
+						// error: scalar property not found
+						if (Policies.MissingFields != EStructDeserializerErrorPolicies::Ignore)
 						{
-							continue;
+							UE_LOG(LogSerialization, Verbose, TEXT("The property '%s' does not exist"), *PropertyName);
 						}
 
-						if (!Backend.ReadProperty(Property, CurrentState.Property, CurrentState.Data, CurrentState.ArrayIndex))
+						if (Policies.MissingFields == EStructDeserializerErrorPolicies::Error)
 						{
-							UE_LOG(LogSerialization, Verbose, TEXT("The property '%s' could not be read (%s)"), *PropertyName, *Backend.GetDebugString());
+							return false;
 						}
 					}
 				}
@@ -229,12 +246,21 @@ bool FStructDeserializer::Deserialize( void* OutStruct, UStruct& TypeInfo, IStru
 
 		case EStructDeserializerBackendTokens::StructureEnd:
 			{
+				// rehash if value was a map
+				UMapProperty* MapProperty = Cast<UMapProperty>(CurrentState.Property);
+
+				if (MapProperty != nullptr)
+				{			
+					FScriptMapHelper MapHelper(MapProperty, CurrentState.Data);
+					MapHelper.Rehash();
+				}
+
 				if (StateStack.Num() == 0)
 				{
 					return true;
 				}
 
-				CurrentState = StateStack.Pop(/*bAllowShrinking=*/ false);
+				CurrentState = StateStack.Pop(/*bAllowShrinking*/ false);
 			}
 			break;
 
@@ -250,12 +276,12 @@ bool FStructDeserializer::Deserialize( void* OutStruct, UStruct& TypeInfo, IStru
 						continue;
 					}
 
-					// handle structured array element
+					// handle struct element inside array
 					UArrayProperty* ArrayProperty = Cast<UArrayProperty>(CurrentState.Property);
 
 					if (ArrayProperty == nullptr)
 					{
-						UE_LOG(LogSerialization, Verbose, TEXT("Property %s is not an array"));
+						UE_LOG(LogSerialization, Verbose, TEXT("Found unnamed value outside of array"));
 
 						return false;
 					}
@@ -266,17 +292,29 @@ bool FStructDeserializer::Deserialize( void* OutStruct, UStruct& TypeInfo, IStru
 					NewState.Property = ArrayProperty->Inner;
 					NewState.Data = ArrayHelper.GetRawPtr(ArrayIndex);
 				}
+				else if ((CurrentState.Property != nullptr) && (CurrentState.Property->GetClass() == UMapProperty::StaticClass()))
+				{
+					// handle map or struct element inside map
+					UMapProperty* MapProperty = Cast<UMapProperty>(CurrentState.Property);
+					FScriptMapHelper MapHelper(MapProperty, CurrentState.Data);
+					int32 PairIndex = MapHelper.AddDefaultValue_Invalid_NeedsRehash();
+					uint8* PairPtr = MapHelper.GetPairPtr(PairIndex);
+					
+					NewState.Data = PairPtr + MapHelper.MapLayout.ValueOffset;
+					NewState.Property = MapProperty->ValueProp;
+
+					MapProperty->KeyProp->ImportText(*PropertyName, PairPtr + MapProperty->MapLayout.KeyOffset, PPF_None, nullptr);
+				}
 				else
 				{
-					// handle structured property
 					NewState.Property = FindField<UProperty>(CurrentState.TypeInfo, *PropertyName);
 
 					if (NewState.Property == nullptr)
 					{
-						// error: structured property not found
+						// error: map or struct property not found
 						if (Policies.MissingFields != EStructDeserializerErrorPolicies::Ignore)
 						{
-							UE_LOG(LogSerialization, Verbose, TEXT("Structured property '%s' not found"), *PropertyName);
+							UE_LOG(LogSerialization, Verbose, TEXT("Map or struct property '%s' not found"), *PropertyName);
 						}
 
 						if (Policies.MissingFields == EStructDeserializerErrorPolicies::Error)
@@ -284,15 +322,25 @@ bool FStructDeserializer::Deserialize( void* OutStruct, UStruct& TypeInfo, IStru
 							return false;
 						}
 					}
+					else if (NewState.Property->GetClass() == UMapProperty::StaticClass())
+					{
+						// handle map property
+						UMapProperty* MapProperty = Cast<UMapProperty>(NewState.Property);
+						NewState.Data = MapProperty->ContainerPtrToValuePtr<void>(CurrentState.Data, CurrentState.ArrayIndex);
+
+						FScriptMapHelper MapHelper(MapProperty, NewState.Data);
+						MapHelper.EmptyValues();
+					}
 					else
 					{
+						// handle struct property
 						NewState.Data = NewState.Property->ContainerPtrToValuePtr<void>(CurrentState.Data);
 					}
 				}
 
 				if (NewState.Property != nullptr)
 				{
-					// Skip struct property if property filter is set and rejects it
+					// skip struct property if property filter is set and rejects it
 					if (Policies.PropertyFilter && !Policies.PropertyFilter(NewState.Property, CurrentState.Property))
 					{
 						Backend.SkipStructure();
diff --git a/Engine/Source/Runtime/Serialization/Private/Tests/StructSerializerTest.cpp b/Engine/Source/Runtime/Serialization/Private/Tests/StructSerializerTest.cpp
--- a/Engine/Source/Runtime/Serialization/Private/Tests/StructSerializerTest.cpp
+++ b/Engine/Source/Runtime/Serialization/Private/Tests/StructSerializerTest.cpp
@@ -72,8 +72,8 @@ namespace StructSerializerTest
 		Test.TestEqual<TArray<FVector>>(TEXT("Arrays.VectorArray must be the same before and after de-/serialization"), TestStruct.Arrays.VectorArray, TestStruct2.Arrays.VectorArray);
 
 		// test maps
-//		Test.TestTrue(TEXT("Maps.IntToString must be the same before and after de-/serialization"), TestStruct.Maps.IntToString.OrderIndependentCompareEqual(TestStruct2.Maps.IntToString));
-//		Test.TestTrue(TEXT("Maps.StringToVector must be the same before and after de-/serialization"), TestStruct.Maps.StringToVector.OrderIndependentCompareEqual(TestStruct2.Maps.StringToVector));
+		Test.TestTrue(TEXT("Maps.IntToStr must be the same before and after de-/serialization"), TestStruct.Maps.IntToStr.OrderIndependentCompareEqual(TestStruct2.Maps.IntToStr));
+		Test.TestTrue(TEXT("Maps.StrToVec must be the same before and after de-/serialization"), TestStruct.Maps.StrToVec.OrderIndependentCompareEqual(TestStruct2.Maps.StrToVec));
 	}
 }
 
diff --git a/Engine/Source/Runtime/Serialization/Private/Tests/StructSerializerTestTypes.h b/Engine/Source/Runtime/Serialization/Private/Tests/StructSerializerTestTypes.h
--- a/Engine/Source/Runtime/Serialization/Private/Tests/StructSerializerTestTypes.h
+++ b/Engine/Source/Runtime/Serialization/Private/Tests/StructSerializerTestTypes.h
@@ -212,21 +212,21 @@ struct FStructSerializerMapTestStruct
 	GENERATED_USTRUCT_BODY()
 
 	UPROPERTY()
-	TMap<int32, FString> IntToString;
+	TMap<int32, FString> IntToStr;
 
 	UPROPERTY()
-	TMap<FString, FVector> StringToVector;
+	TMap<FString, FVector> StrToVec;
 
 	/** Default constructor. */
 	FStructSerializerMapTestStruct()
 	{
-		IntToString.Add(1, TEXT("One"));
-		IntToString.Add(2, TEXT("Two"));
-		IntToString.Add(3, TEXT("Three"));
+		IntToStr.Add(1, TEXT("One"));
+		IntToStr.Add(2, TEXT("Two"));
+		IntToStr.Add(3, TEXT("Three"));
 
-		StringToVector.Add(TEXT("Vector000"), FVector(0.0f, 0.0f, 0.0f));
-		StringToVector.Add(TEXT("Vector123"), FVector(1.0f, 2.0f, 3.0f));
-		StringToVector.Add(TEXT("Vector666"), FVector(6.0f, 6.0f, 6.0f));
+		StrToVec.Add(TEXT("V000"), FVector(0.0f, 0.0f, 0.0f));
+		StrToVec.Add(TEXT("V123"), FVector(1.0f, 2.0f, 3.0f));
+		StrToVec.Add(TEXT("V666"), FVector(6.0f, 6.0f, 6.0f));
 	}
 
 	/** Creates an uninitialized instance. */
