Copying //UE4/Dev-Core to //UE4/Dev-Main (Source: //UE4/Dev-Core @ 3208226)

diff --git a/Engine/Source/Runtime/Renderer/Private/DistanceFieldGlobalIllumination.cpp b/Engine/Source/Runtime/Renderer/Private/DistanceFieldGlobalIllumination.cpp
--- a/Engine/Source/Runtime/Renderer/Private/DistanceFieldGlobalIllumination.cpp
+++ b/Engine/Source/Runtime/Renderer/Private/DistanceFieldGlobalIllumination.cpp
@@ -16,6 +16,7 @@
 #include "RHICommandList.h"
 #include "SceneUtils.h"
 #include "DistanceFieldAtlas.h"
+#include "UniquePtr.h"
 
 int32 GDistanceFieldGI = 0;
 FAutoConsoleVariableRef CVarDistanceFieldGI(
@@ -376,7 +377,7 @@ private:
 
 IMPLEMENT_SHADER_TYPE(,FCullVPLsForViewCS,TEXT("DistanceFieldGlobalIllumination"),TEXT("CullVPLsForViewCS"),SF_Compute);
 
-TScopedPointer<FLightTileIntersectionResources> GVPLPlacementTileIntersectionResources;
+TUniquePtr<FLightTileIntersectionResources> GVPLPlacementTileIntersectionResources;
 
 void PlaceVPLs(
 	FRHICommandListImmediate& RHICmdList,
@@ -526,7 +527,7 @@ void PlaceVPLs(
 			TShaderMapRef<FVPLPlacementCS> ComputeShader(View.ShaderMap);
 
 			RHICmdList.SetComputeShader(ComputeShader->GetComputeShader());
-			ComputeShader->SetParameters(RHICmdList, View, DirectionalLightProxy, FVector2D(1.0f / GVPLGridDimension, 1.0f / GVPLGridDimension), DirectionalLightShadowToWorld, DirectionalLightShadowToWorld.InverseFast(), GVPLPlacementTileIntersectionResources);
+			ComputeShader->SetParameters(RHICmdList, View, DirectionalLightProxy, FVector2D(1.0f / GVPLGridDimension, 1.0f / GVPLGridDimension), DirectionalLightShadowToWorld, DirectionalLightShadowToWorld.InverseFast(), GVPLPlacementTileIntersectionResources.Get());
 			DispatchComputeShader(RHICmdList, *ComputeShader, FMath::DivideAndRoundUp<int32>(GVPLGridDimension, GDistanceFieldAOTileSizeX), FMath::DivideAndRoundUp<int32>(GVPLGridDimension, GDistanceFieldAOTileSizeY), 1);
 
 			ComputeShader->UnsetParameters(RHICmdList);
@@ -880,7 +881,7 @@ void UpdateVPLs(
 				{
 					TShaderMapRef<FLightVPLsCS> ComputeShader(View.ShaderMap);
 					RHICmdList.SetComputeShader(ComputeShader->GetComputeShader());
-					ComputeShader->SetParameters(RHICmdList, View, DirectionalLightProxy, DirectionalLightWorldToShadow, Parameters, GVPLPlacementTileIntersectionResources);
+					ComputeShader->SetParameters(RHICmdList, View, DirectionalLightProxy, DirectionalLightWorldToShadow, Parameters, GVPLPlacementTileIntersectionResources.Get());
 					DispatchIndirectComputeShader(RHICmdList, *ComputeShader, GAOCulledObjectBuffers.Buffers.ObjectIndirectDispatch.Buffer, 0);
 					ComputeShader->UnsetParameters(RHICmdList);
 				}
diff --git a/Engine/Source/Runtime/Renderer/Private/DistanceFieldLightingShared.h b/Engine/Source/Runtime/Renderer/Private/DistanceFieldLightingShared.h
--- a/Engine/Source/Runtime/Renderer/Private/DistanceFieldLightingShared.h
+++ b/Engine/Source/Runtime/Renderer/Private/DistanceFieldLightingShared.h
@@ -7,6 +7,7 @@
 #pragma once
 
 #include "DistanceFieldAtlas.h"
+#include "UniquePtr.h"
 
 DECLARE_LOG_CATEGORY_EXTERN(LogDistanceField, Warning, All);
 
@@ -489,7 +490,7 @@ extern void CullDistanceFieldObjectsForLight(
 	const FPlane* PlaneData, 
 	const FVector4& ShadowBoundingSphereValue,
 	float ShadowBoundingRadius,
-	TScopedPointer<class FLightTileIntersectionResources>& TileIntersectionResources);
+	TUniquePtr<class FLightTileIntersectionResources>& TileIntersectionResources);
 
 class FUniformMeshBuffers
 {
diff --git a/Engine/Source/Runtime/Renderer/Private/DistanceFieldShadowing.cpp b/Engine/Source/Runtime/Renderer/Private/DistanceFieldShadowing.cpp
--- a/Engine/Source/Runtime/Renderer/Private/DistanceFieldShadowing.cpp
+++ b/Engine/Source/Runtime/Renderer/Private/DistanceFieldShadowing.cpp
@@ -729,7 +729,7 @@ void CullDistanceFieldObjectsForLight(
 	const FPlane* PlaneData, 
 	const FVector4& ShadowBoundingSphereValue,
 	float ShadowBoundingRadius,
-	TScopedPointer<FLightTileIntersectionResources>& TileIntersectionResources)
+	TUniquePtr<FLightTileIntersectionResources>& TileIntersectionResources)
 {
 	const FScene* Scene = (const FScene*)(View.Family->Scene);
 
@@ -773,7 +773,7 @@ void CullDistanceFieldObjectsForLight(
 			}
 			else
 			{
-				TileIntersectionResources = new FLightTileIntersectionResources();
+				TileIntersectionResources = MakeUnique<FLightTileIntersectionResources>();
 			}
 
 			TileIntersectionResources->TileDimensions = LightTileDimensions;
@@ -788,10 +788,10 @@ void CullDistanceFieldObjectsForLight(
 			TShaderMapRef<FWorkaroundAMDBugCS> ComputeShader(View.ShaderMap);
 
 			RHICmdList.SetComputeShader(ComputeShader->GetComputeShader());
-			ComputeShader->SetParameters(RHICmdList, View, TileIntersectionResources);
+			ComputeShader->SetParameters(RHICmdList, View, TileIntersectionResources.Get());
 			DispatchComputeShader(RHICmdList, *ComputeShader, 1, 1, 1);
 
-			ComputeShader->UnsetParameters(RHICmdList, TileIntersectionResources);
+			ComputeShader->UnsetParameters(RHICmdList, TileIntersectionResources.Get());
 		}
 		
 		{
@@ -801,10 +801,10 @@ void CullDistanceFieldObjectsForLight(
 			uint32 GroupSizeY = FMath::DivideAndRoundUp(LightTileDimensions.Y, GDistanceFieldAOTileSizeY);
 
 			RHICmdList.SetComputeShader(ComputeShader->GetComputeShader());
-			ComputeShader->SetParameters(RHICmdList, View, FVector2D(LightTileDimensions.X, LightTileDimensions.Y), TileIntersectionResources);
+			ComputeShader->SetParameters(RHICmdList, View, FVector2D(LightTileDimensions.X, LightTileDimensions.Y), TileIntersectionResources.Get());
 			DispatchComputeShader(RHICmdList, *ComputeShader, GroupSizeX, GroupSizeY, 1);
 
-			ComputeShader->UnsetParameters(RHICmdList, TileIntersectionResources);
+			ComputeShader->UnsetParameters(RHICmdList, TileIntersectionResources.Get());
 		}
 
 		{
@@ -812,7 +812,7 @@ void CullDistanceFieldObjectsForLight(
 			TShaderMapRef<FShadowObjectCullPS> PixelShader(View.ShaderMap);
 
 			TArray<FUnorderedAccessViewRHIParamRef> UAVs;
-			PixelShader->GetUAVs(View, TileIntersectionResources, UAVs);
+			PixelShader->GetUAVs(View, TileIntersectionResources.Get(), UAVs);
 			RHICmdList.TransitionResources(EResourceTransitionAccess::ERWBarrier, EResourceTransitionPipeline::EComputeToGfx, UAVs.GetData(), UAVs.Num());
 			RHICmdList.SetRenderTargets(0, (const FRHIRenderTargetView *)NULL, NULL, UAVs.Num(), UAVs.GetData());
 
@@ -932,7 +932,7 @@ void FProjectedShadowInfo::RenderRayTracedDistanceFieldProjection(FRHICommandLis
 				LightSceneInfo->TileIntersectionResources
 				);
 
-			FLightTileIntersectionResources* TileIntersectionResources = LightSceneInfo->TileIntersectionResources;
+			FLightTileIntersectionResources* TileIntersectionResources = LightSceneInfo->TileIntersectionResources.Get();
 
 			View.HeightfieldLightingViewInfo.ComputeRayTracedShadowing(View, RHICmdList, this, TileIntersectionResources, GShadowCulledObjectBuffers);
 
diff --git a/Engine/Source/Runtime/Renderer/Private/LightSceneInfo.cpp b/Engine/Source/Runtime/Renderer/Private/LightSceneInfo.cpp
--- a/Engine/Source/Runtime/Renderer/Private/LightSceneInfo.cpp
+++ b/Engine/Source/Runtime/Renderer/Private/LightSceneInfo.cpp
@@ -39,7 +39,7 @@ FLightSceneInfo::FLightSceneInfo(FLightSceneProxy* InProxy, bool InbVisible)
 	, DynamicInteractionOftenMovingPrimitiveList(NULL)
 	, DynamicInteractionStaticPrimitiveList(NULL)
 	, Id(INDEX_NONE)
-	, TileIntersectionResources(NULL)
+	, TileIntersectionResources(nullptr)
 	, bPrecomputedLightingIsValid(InProxy->GetLightComponent()->IsPrecomputedLightingValid())
 	, bVisible(InbVisible)
 	, bEnableLightShaftBloom(InProxy->GetLightComponent()->bEnableLightShaftBloom)
diff --git a/Engine/Source/Runtime/Renderer/Private/LightSceneInfo.h b/Engine/Source/Runtime/Renderer/Private/LightSceneInfo.h
--- a/Engine/Source/Runtime/Renderer/Private/LightSceneInfo.h
+++ b/Engine/Source/Runtime/Renderer/Private/LightSceneInfo.h
@@ -8,6 +8,7 @@
 
 #include "StaticArray.h"
 #include "SceneManagement.h"
+#include "UniquePtr.h"
 
 /**
  * The information needed to cull a light-primitive interaction.
@@ -120,7 +121,7 @@ public:
 	mutable const FMaterialShaderMap* TranslucentInjectCachedShaderMaps[LightType_MAX][2][2][2];
 
 	/** Tile intersection buffer for distance field shadowing, stored on the light to avoid reallocating each frame. */
-	mutable TScopedPointer<class FLightTileIntersectionResources> TileIntersectionResources;
+	mutable TUniquePtr<class FLightTileIntersectionResources> TileIntersectionResources;
 
 	mutable FVertexBufferRHIRef ShadowCapsuleShapesVertexBuffer;
 	mutable FShaderResourceViewRHIRef ShadowCapsuleShapesSRV;
diff --git a/Engine/Source/Runtime/Renderer/Private/ReflectionEnvironment.cpp b/Engine/Source/Runtime/Renderer/Private/ReflectionEnvironment.cpp
--- a/Engine/Source/Runtime/Renderer/Private/ReflectionEnvironment.cpp
+++ b/Engine/Source/Runtime/Renderer/Private/ReflectionEnvironment.cpp
@@ -184,6 +184,125 @@ void FReflectionEnvironmentCubemapArray::ReleaseDynamicRHI()
 	ReleaseCubeArray();
 }
 
+void FReflectionEnvironmentSceneData::ResizeCubemapArrayGPU(uint32 InMaxCubemaps, int32 InCubemapSize)
+{
+	check(IsInRenderingThread());
+
+	// If the cubemap array isn't setup yet then no copying/reallocation is necessary. Just go through the old path
+	if (!CubemapArray.IsInitialized())
+	{
+		CubemapArray.UpdateMaxCubemaps(InMaxCubemaps, InCubemapSize);
+		return;
+	}
+
+	// Generate a remapping table for the elements
+	TArray<bool> ArrayIndicesRemoved;
+	ArrayIndicesRemoved.Empty(CubemapArray.GetMaxCubemaps());
+	for (int i = 0; i < CubemapArray.GetMaxCubemaps(); i++)
+	{
+		ArrayIndicesRemoved.Add(false);
+	}
+	for (int i = 0; i < CubemapIndicesRemovedSinceLastRealloc.Num(); i++)
+	{
+		uint32 CubemapIndex = CubemapIndicesRemovedSinceLastRealloc[i];
+		ArrayIndicesRemoved[CubemapIndex] = true;
+	}
+	TArray<int32> IndexRemapping;
+	int32 Count = 0;
+	for (int i = 0; i < CubemapArray.GetMaxCubemaps(); i++)
+	{
+		if (ArrayIndicesRemoved[i])
+		{
+			IndexRemapping.Add(-1);
+		}
+		else
+		{
+			IndexRemapping.Add(Count);
+			Count++;
+		}
+	}
+
+	// Spin through the AllocatedReflectionCaptureState map and remap the indices based on the LUT
+	TArray<const UReflectionCaptureComponent*> Components;
+	AllocatedReflectionCaptureState.GetKeys(Components);
+	int32 UsedCubemapCount = 0;
+	for (int32 i=0; i<Components.Num(); i++ )
+	{
+		FCaptureComponentSceneState* ComponentStatePtr = AllocatedReflectionCaptureState.Find(Components[i]);
+		check(ComponentStatePtr->CaptureIndex < IndexRemapping.Num());
+		ComponentStatePtr->CaptureIndex = IndexRemapping[ComponentStatePtr->CaptureIndex];
+		UsedCubemapCount = FMath::Max(UsedCubemapCount, ComponentStatePtr->CaptureIndex + 1);
+	}
+
+	// Clear elements in the remapping array which are outside the range of the used components (these were allocated but not used)
+	for (int i = 0; i < IndexRemapping.Num(); i++)
+	{
+		if (IndexRemapping[i] >= UsedCubemapCount)
+		{
+			IndexRemapping[i] = -1;
+		}
+	}
+
+	CubemapArray.ResizeCubemapArrayGPU(InMaxCubemaps, InCubemapSize, IndexRemapping);
+	CubemapIndicesRemovedSinceLastRealloc.Empty();
+}
+
+void FReflectionEnvironmentCubemapArray::ResizeCubemapArrayGPU(uint32 InMaxCubemaps, int32 InCubemapSize, const TArray<int32>& IndexRemapping)
+{
+	check(IsInRenderingThread());
+	check(GetFeatureLevel() >= ERHIFeatureLevel::SM5);
+	check(IsInitialized());
+	check(InCubemapSize == CubemapSize);
+
+	// Take a reference to the old cubemap array and then release it to prevent it getting destroyed during InitDynamicRHI
+	TRefCountPtr<IPooledRenderTarget> OldReflectionEnvs = ReflectionEnvs;
+	ReflectionEnvs = nullptr;
+	int OldMaxCubemaps = MaxCubemaps;
+	MaxCubemaps = InMaxCubemaps;
+
+	InitDynamicRHI();
+
+	FTextureRHIRef TexRef = OldReflectionEnvs->GetRenderTargetItem().TargetableTexture;
+	FRHICommandListImmediate& RHICmdList = FRHICommandListExecutor::GetImmediateCommandList();
+	const int32 NumMips = FMath::CeilLogTwo(InCubemapSize) + 1;
+
+	{
+		SCOPED_DRAW_EVENT(RHICmdList, ReflectionEnvironment_ResizeCubemapArray);
+		SCOPED_GPU_STAT(RHICmdList, Stat_GPU_ReflectionEnvironment)
+
+		// Copy the cubemaps, remapping the elements as necessary
+		FResolveParams ResolveParams;
+		ResolveParams.Rect = FResolveRect();
+		for (int32 SourceCubemapIndex = 0; SourceCubemapIndex < OldMaxCubemaps; SourceCubemapIndex++)
+		{
+			int32 DestCubemapIndex = IndexRemapping[SourceCubemapIndex];
+			if (DestCubemapIndex != -1)
+			{
+				ResolveParams.SourceArrayIndex = SourceCubemapIndex;
+				ResolveParams.DestArrayIndex = DestCubemapIndex;
+
+				check(SourceCubemapIndex < OldMaxCubemaps);
+				check(DestCubemapIndex < (int32)MaxCubemaps);
+
+				for (int Face = 0; Face < 6; Face++)
+				{
+					ResolveParams.CubeFace = (ECubeFace)Face;
+					for (int Mip = 0; Mip < NumMips; Mip++)
+					{
+						ResolveParams.MipIndex = Mip;
+						//@TODO: We should use an explicit copy method for this rather than CopyToResolveTarget, but that doesn't exist right now. 
+						// For now, we'll just do this on RHIs where we know CopyToResolveTarget does the right thing. In future we should look to 
+						// add a a new RHI method
+						check(GRHISupportsResolveCubemapFaces);
+						RHICmdList.CopyToResolveTarget(OldReflectionEnvs->GetRenderTargetItem().ShaderResourceTexture, ReflectionEnvs->GetRenderTargetItem().ShaderResourceTexture, true, ResolveParams);
+					}
+				}
+			}
+		}
+	}
+	GRenderTargetPool.FreeUnusedResource(OldReflectionEnvs);
+}
+
 void FReflectionEnvironmentCubemapArray::UpdateMaxCubemaps(uint32 InMaxCubemaps, int32 InCubemapSize)
 {
 	MaxCubemaps = InMaxCubemaps;
diff --git a/Engine/Source/Runtime/Renderer/Private/ReflectionEnvironmentCapture.cpp b/Engine/Source/Runtime/Renderer/Private/ReflectionEnvironmentCapture.cpp
--- a/Engine/Source/Runtime/Renderer/Private/ReflectionEnvironmentCapture.cpp
+++ b/Engine/Source/Runtime/Renderer/Private/ReflectionEnvironmentCapture.cpp
@@ -49,6 +49,19 @@ FAutoConsoleCommandWithWorld CaptureConsoleCommand(
 	FConsoleCommandWithWorldDelegate::CreateStatic(OnUpdateReflectionCaptures)
 	);
 
+static TAutoConsoleVariable<int32> CVarReflectionCaptureGPUArrayCopy(
+	TEXT("r.ReflectionCaptureGPUArrayCopy"),
+	1,
+	TEXT("Do a fast copy of the reflection capture array when resizing if possible. This avoids hitches on the rendering thread when the cubemap array needs to grow.\n")
+	TEXT(" 0 is off, 1 is on (default)"),
+	ECVF_ReadOnly);
+
+
+bool DoGPUArrayCopy()
+{
+	return GRHISupportsResolveCubemapFaces && CVarReflectionCaptureGPUArrayCopy.GetValueOnAnyThread();
+}
+
 /** Encapsulates render target picking logic for cubemap mip generation. */
 FSceneRenderTargetItem& GetEffectiveRenderTarget(FSceneRenderTargets& SceneContext, bool bDownsamplePass, int32 TargetMipIndex)
 {
@@ -963,7 +976,30 @@ void FScene::AllocateReflectionCaptures(const TArray<UReflectionCaptureComponent
 			DesiredMaxCubemaps = FMath::Min(DesiredMaxCubemaps, GMaxNumReflectionCaptures);
 
 			const int32 ReflectionCaptureSize = UReflectionCaptureComponent::GetReflectionCaptureSize_GameThread();
-			if (DesiredMaxCubemaps != ReflectionSceneData.MaxAllocatedReflectionCubemapsGameThread || ReflectionCaptureSize != ReflectionSceneData.CubemapArray.GetCubemapSize())
+			bool bNeedsUpdateAllCaptures = DesiredMaxCubemaps != ReflectionSceneData.MaxAllocatedReflectionCubemapsGameThread || ReflectionCaptureSize != ReflectionSceneData.CubemapArray.GetCubemapSize();
+
+			if (DoGPUArrayCopy() && bNeedsUpdateAllCaptures)
+			{
+				check( ReflectionCaptureSize == ReflectionSceneData.CubemapArray.GetCubemapSize() || ReflectionSceneData.CubemapArray.GetCubemapSize() == 0 );
+				if (ReflectionCaptureSize == ReflectionSceneData.CubemapArray.GetCubemapSize())
+				{
+					// We can do a fast GPU copy to realloc the array, so we don't need to update all captures
+					ReflectionSceneData.MaxAllocatedReflectionCubemapsGameThread = DesiredMaxCubemaps;
+					ENQUEUE_UNIQUE_RENDER_COMMAND_THREEPARAMETER(
+						GPUResizeArrayCommand,
+						FScene*, Scene, this,
+						uint32, MaxSize, ReflectionSceneData.MaxAllocatedReflectionCubemapsGameThread,
+						int32, ReflectionCaptureSize, ReflectionCaptureSize,
+						{
+							// Update the scene's cubemap array, preserving the original contents with a GPU-GPU copy
+							Scene->ReflectionSceneData.ResizeCubemapArrayGPU(MaxSize, ReflectionCaptureSize);
+						});
+
+					bNeedsUpdateAllCaptures = false;
+				}
+			}
+
+			if (bNeedsUpdateAllCaptures)
 			{
 				ReflectionSceneData.MaxAllocatedReflectionCubemapsGameThread = DesiredMaxCubemaps;
 
@@ -982,8 +1018,7 @@ void FScene::AllocateReflectionCaptures(const TArray<UReflectionCaptureComponent
 			}
 			else
 			{
-				// No reallocation of the cubemap array was needed, just update the captures that were requested
-
+				// No teardown of the cubemap array was needed, just update the captures that were requested
 				for (TSparseArray<UReflectionCaptureComponent*>::TIterator It(ReflectionSceneData.AllocatedReflectionCapturesGameThread); It; ++It)
 				{
 					UReflectionCaptureComponent* CurrentComponent = *It;
@@ -1031,6 +1066,7 @@ void FScene::UpdateAllReflectionCaptures()
 			FScene*, Scene, this,
 		{
 			Scene->ReflectionSceneData.AllocatedReflectionCaptureState.Empty();
+			Scene->ReflectionSceneData.CubemapIndicesRemovedSinceLastRealloc.Empty();
 		});
 
 		const int32 UpdateDivisor = FMath::Max(ReflectionSceneData.AllocatedReflectionCapturesGameThread.Num() / 20, 1);
@@ -1146,8 +1182,7 @@ void UploadReflectionCapture_RenderingThread(FScene* Scene, const FReflectionCap
 	FTextureCubeRHIRef& CubeMapArray = (FTextureCubeRHIRef&)Scene->ReflectionSceneData.CubemapArray.GetRenderTarget().ShaderResourceTexture;
 	check(CubeMapArray->GetFormat() == PF_FloatRGBA);
 
-	TArray<uint8> CubemapData;
-	FullHDRData->GetUncompressedData(CubemapData);
+	TRefCountPtr<FReflectionCaptureUncompressedData> SourceCubemapData = FullHDRData->GetUncompressedData();
 	int32 MipBaseIndex = 0;
 
 	for (int32 MipIndex = 0; MipIndex < NumMips; MipIndex++)
@@ -1165,7 +1200,7 @@ void UploadReflectionCapture_RenderingThread(FScene* Scene, const FReflectionCap
 			{
 				FFloat16Color* DestPtr = (FFloat16Color*)((uint8*)DestBuffer + Y * DestStride);
 				const int32 SourceIndex = MipBaseIndex + CubeFace * CubeFaceBytes + Y * MipSize * sizeof(FFloat16Color);
-				const uint8* SourcePtr = &CubemapData[SourceIndex];
+				const uint8* SourcePtr = SourceCubemapData->GetData(SourceIndex);
 				FMemory::Memcpy(DestPtr, SourcePtr, MipSize * sizeof(FFloat16Color));
 			}
 
@@ -1426,7 +1461,7 @@ void FScene::UpdateReflectionCaptureContents(UReflectionCaptureComponent* Captur
 		const FReflectionCaptureFullHDR* DerivedData = CaptureComponent->GetFullHDRData();
 
 		// Upload existing derived data if it exists, instead of capturing
-		if (DerivedData && DerivedData->CompressedCapturedData.Num() > 0)
+		if (DerivedData && DerivedData->HasValidData() )
 		{
 			// For other feature levels the reflection textures are stored on the component instead of in a scene-wide texture array
 			if (GetFeatureLevel() >= ERHIFeatureLevel::SM5)
@@ -1439,6 +1474,14 @@ void FScene::UpdateReflectionCaptureContents(UReflectionCaptureComponent* Captur
 				{
 					UploadReflectionCapture_RenderingThread(Scene, DerivedData, CaptureComponent);
 				});
+
+				if ( DoGPUArrayCopy() && !GIsEditor )
+				{
+					// We no longer need the HDR data, since we have a copy on the GPU
+					// In the editor we need this data for serialization, however
+					DerivedData = nullptr;
+					CaptureComponent->ReleaseHDRData();
+				}
 			}
 		}
 		else
diff --git a/Engine/Source/Runtime/Renderer/Private/RendererScene.cpp b/Engine/Source/Runtime/Renderer/Private/RendererScene.cpp
--- a/Engine/Source/Runtime/Renderer/Private/RendererScene.cpp
+++ b/Engine/Source/Runtime/Renderer/Private/RendererScene.cpp
@@ -1310,6 +1310,16 @@ void FScene::ReleaseReflectionCubemap(UReflectionCaptureComponent* CaptureCompon
 		UReflectionCaptureComponent*,Component,CaptureComponent,
 		FScene*,Scene,this,
 	{
+		// In SM5 we track removed captures so we can remap them when reallocating the cubemap array
+		if (Scene->GetFeatureLevel() >= ERHIFeatureLevel::SM5)
+		{
+			const FCaptureComponentSceneState* ComponentStatePtr = Scene->ReflectionSceneData.AllocatedReflectionCaptureState.Find(Component);
+			if (ComponentStatePtr)
+			{
+				Scene->ReflectionSceneData.CubemapIndicesRemovedSinceLastRealloc.Add(ComponentStatePtr->CaptureIndex);
+			}
+		}
+
 		Scene->ReflectionSceneData.AllocatedReflectionCaptureState.Remove(Component);
 	});
 }
diff --git a/Engine/Source/Runtime/Renderer/Private/SceneCore.cpp b/Engine/Source/Runtime/Renderer/Private/SceneCore.cpp
--- a/Engine/Source/Runtime/Renderer/Private/SceneCore.cpp
+++ b/Engine/Source/Runtime/Renderer/Private/SceneCore.cpp
@@ -384,7 +384,7 @@ void FStaticMesh::RemoveFromDrawLists()
 		FStaticMesh::FDrawListElementLink* Link = DrawListLinks[0];
 		const int32 OriginalNumLinks = DrawListLinks.Num();
 		// This will call UnlinkDrawList.
-		Link->Remove();
+		Link->Remove(true);
 		check(DrawListLinks.Num() == OriginalNumLinks - 1);
 		if(DrawListLinks.Num())
 		{
@@ -411,7 +411,12 @@ FStaticMesh::~FStaticMesh()
 	// Remove this static mesh from the scene's list.
 	PrimitiveSceneInfo->Scene->StaticMeshes.RemoveAt(Id);
 
-	RemoveFromDrawLists();
+	// This is cheaper than calling RemoveFromDrawLists, since it 
+	// doesn't unlink meshes which are about to be destroyed
+	for (int32 i = 0; i < DrawListLinks.Num(); i++)
+	{
+		DrawListLinks[i]->Remove(false);
+	}
 }
 
 /** Initialization constructor. */
diff --git a/Engine/Source/Runtime/Renderer/Private/SceneCore.h b/Engine/Source/Runtime/Renderer/Private/SceneCore.h
--- a/Engine/Source/Runtime/Renderer/Private/SceneCore.h
+++ b/Engine/Source/Runtime/Renderer/Private/SceneCore.h
@@ -131,7 +131,7 @@ public:
 	{
 	public:
 		virtual bool IsInDrawList(const class FStaticMeshDrawListBase* DrawList) const = 0;
-		virtual void Remove() = 0;
+		virtual void Remove(const bool bUnlinkMesh) = 0;
 	};
 
 	/** The screen space size to draw this primitive at */
diff --git a/Engine/Source/Runtime/Renderer/Private/ScenePrivate.h b/Engine/Source/Runtime/Renderer/Private/ScenePrivate.h
--- a/Engine/Source/Runtime/Renderer/Private/ScenePrivate.h
+++ b/Engine/Source/Runtime/Renderer/Private/ScenePrivate.h
@@ -1102,6 +1102,13 @@ public:
 	 * This reallocates the resource but does not copy over the old contents. 
 	 */
 	void UpdateMaxCubemaps(uint32 InMaxCubemaps, int32 CubemapSize);
+
+	/**
+	* Updates the maximum number of cubemaps that this array is allocated for.
+	* This reallocates the resource and copies over the old contents, preserving indices
+	*/
+	void ResizeCubemapArrayGPU(uint32 InMaxCubemaps, int32 CubemapSize, const TArray<int32>& IndexRemapping);
+
 	int32 GetMaxCubemaps() const { return MaxCubemaps; }
 	int32 GetCubemapSize() const { return CubemapSize; }
 	bool IsValid() const { return IsValidRef(ReflectionEnvs); }
@@ -1153,6 +1160,9 @@ public:
 	 */
 	FReflectionEnvironmentCubemapArray CubemapArray;
 
+	/** We track the cubemaps removed since the last reallocation to allow us to remap them reallocating the array */
+	TArray<uint32> CubemapIndicesRemovedSinceLastRealloc;
+
 	/** Rendering thread map from component to scene state.  This allows storage of RT state that needs to persist through a component re-register. */
 	TMap<const UReflectionCaptureComponent*, FCaptureComponentSceneState> AllocatedReflectionCaptureState;
 
@@ -1170,6 +1180,9 @@ public:
 		CubemapArray(InFeatureLevel),
 		MaxAllocatedReflectionCubemapsGameThread(0)
 	{}
+
+
+	void ResizeCubemapArrayGPU(uint32 InMaxCubemaps, int32 InCubemapSize);
 };
 
 class FPrimitiveAndInstance
diff --git a/Engine/Source/Runtime/Renderer/Private/SceneRendering.cpp b/Engine/Source/Runtime/Renderer/Private/SceneRendering.cpp
--- a/Engine/Source/Runtime/Renderer/Private/SceneRendering.cpp
+++ b/Engine/Source/Runtime/Renderer/Private/SceneRendering.cpp
@@ -375,7 +375,7 @@ FViewInfo::FViewInfo(const FSceneView* InView)
 
 void FViewInfo::Init()
 {
-	CachedViewUniformShaderParameters = NULL;
+	CachedViewUniformShaderParameters = nullptr;
 	bHasTranslucentViewMeshElements = 0;
 	bPrevTransformsReset = false;
 	bIgnoreExistingQueries = false;
@@ -824,7 +824,7 @@ void FViewInfo::InitRHIResources()
 
 	check(IsInRenderingThread());
 
-	CachedViewUniformShaderParameters = new FViewUniformShaderParameters();
+	CachedViewUniformShaderParameters = MakeUnique<FViewUniformShaderParameters>();
 
 	FSceneRenderTargets& SceneContext = FSceneRenderTargets::Get(FRHICommandListExecutor::GetImmediateCommandList());
 
@@ -879,7 +879,7 @@ FViewInfo* FViewInfo::CreateSnapshot() const
 	for (size_t i = 0; i < ARRAY_COUNT(Result->MobileDirectionalLightUniformBuffers); i++)
 		FMemory::Memcpy(Result->MobileDirectionalLightUniformBuffers[i], NullMobileDirectionalLightUniformBuffer);
 
-	TScopedPointer<FViewUniformShaderParameters> NullViewParameters;
+	TUniquePtr<FViewUniformShaderParameters> NullViewParameters;
 	FMemory::Memcpy(Result->CachedViewUniformShaderParameters, NullViewParameters); 
 	Result->bIsSnapshot = true;
 	ViewInfoSnapshots.Add(Result);
diff --git a/Engine/Source/Runtime/Renderer/Private/SceneRendering.h b/Engine/Source/Runtime/Renderer/Private/SceneRendering.h
--- a/Engine/Source/Runtime/Renderer/Private/SceneRendering.h
+++ b/Engine/Source/Runtime/Renderer/Private/SceneRendering.h
@@ -11,6 +11,7 @@
 #include "CustomDepthRendering.h"
 #include "HeightfieldLighting.h"
 #include "GlobalDistanceFieldParameters.h"
+#include "UniquePtr.h"
 
 // Forward declarations.
 class FPostprocessContext;
@@ -671,7 +672,7 @@ public:
 	FSceneViewState* ViewState;
 
 	/** Cached view uniform shader parameters, to allow recreating the view uniform buffer without having to fill out the entire struct. */
-	TScopedPointer<FViewUniformShaderParameters> CachedViewUniformShaderParameters;
+	TUniquePtr<FViewUniformShaderParameters> CachedViewUniformShaderParameters;
 
 	/** A map from primitive ID to a boolean visibility value. */
 	FSceneBitArray PrimitiveVisibilityMap;
diff --git a/Engine/Source/Runtime/Renderer/Private/ShadowDepthRendering.cpp b/Engine/Source/Runtime/Renderer/Private/ShadowDepthRendering.cpp
--- a/Engine/Source/Runtime/Renderer/Private/ShadowDepthRendering.cpp
+++ b/Engine/Source/Runtime/Renderer/Private/ShadowDepthRendering.cpp
@@ -1879,7 +1879,7 @@ void FProjectedShadowInfo::ModifyViewForShadow(FRHICommandList& RHICmdList, FVie
 	FoundView->ViewMatrices.HackRemoveTemporalAAProjectionJitter();
 
 	FSceneRenderTargets& SceneContext = FSceneRenderTargets::Get(RHICmdList);
-	FoundView->CachedViewUniformShaderParameters = new FViewUniformShaderParameters();
+	FoundView->CachedViewUniformShaderParameters = MakeUnique<FViewUniformShaderParameters>();
 
 	// Override the view matrix so that billboarding primitives will be aligned to the light
 	FoundView->ViewMatrices.HackOverrideViewMatrixForShadows(ShadowViewMatrix);
diff --git a/Engine/Source/Runtime/Renderer/Private/StaticMeshDrawList.h b/Engine/Source/Runtime/Renderer/Private/StaticMeshDrawList.h
--- a/Engine/Source/Runtime/Renderer/Private/StaticMeshDrawList.h
+++ b/Engine/Source/Runtime/Renderer/Private/StaticMeshDrawList.h
@@ -147,7 +147,7 @@ private:
 			return DrawList == StaticMeshDrawList;
 		}
 		// FAbstractDrawListElementLink interface.
-		virtual void Remove();
+		virtual void Remove(const bool bUnlinkMesh = true );
 
 	private:
 		TStaticMeshDrawList* StaticMeshDrawList;
diff --git a/Engine/Source/Runtime/Renderer/Private/StaticMeshDrawList.inl b/Engine/Source/Runtime/Renderer/Private/StaticMeshDrawList.inl
--- a/Engine/Source/Runtime/Renderer/Private/StaticMeshDrawList.inl
+++ b/Engine/Source/Runtime/Renderer/Private/StaticMeshDrawList.inl
@@ -14,7 +14,7 @@
 #define PER_MESH_DRAW_STATS 0
 
 template<typename DrawingPolicyType>
-void TStaticMeshDrawList<DrawingPolicyType>::FElementHandle::Remove()
+void TStaticMeshDrawList<DrawingPolicyType>::FElementHandle::Remove(const bool bUnlinkMesh)
 {
 	check(!GDrawListsLocked);
 	// Make a copy of this handle's variables on the stack, since the call to Elements.RemoveSwap deletes the handle.
@@ -27,8 +27,12 @@ void TStaticMeshDrawList<DrawingPolicyType>::FElementHandle::Remove()
 	check(LocalDrawingPolicyLink->Elements[ElementIndex].Mesh->MaterialRenderProxy);
 	LocalDrawingPolicyLink->Elements[ElementIndex].Mesh->MaterialRenderProxy->SetUnreferencedInDrawList();
 
-	// Unlink the mesh from this draw list.
-	LocalDrawingPolicyLink->Elements[ElementIndex].Mesh->UnlinkDrawList(this);
+	// Unlink the mesh from this draw list. Not necessary if the mesh is being destroyed
+	if (bUnlinkMesh)
+	{
+		// Expensive (Order N). Spins through whole list
+		LocalDrawingPolicyLink->Elements[ElementIndex].Mesh->UnlinkDrawList(this);
+	}
 	LocalDrawingPolicyLink->Elements[ElementIndex].Mesh = NULL;
 
 	checkSlow(LocalDrawingPolicyLink->Elements.Num() == LocalDrawingPolicyLink->CompactElements.Num());
diff --git a/Engine/Source/Runtime/SandboxFile/Private/IPlatformFileSandboxWrapper.cpp b/Engine/Source/Runtime/SandboxFile/Private/IPlatformFileSandboxWrapper.cpp
--- a/Engine/Source/Runtime/SandboxFile/Private/IPlatformFileSandboxWrapper.cpp
+++ b/Engine/Source/Runtime/SandboxFile/Private/IPlatformFileSandboxWrapper.cpp
@@ -3,6 +3,7 @@
 #include "SandboxFilePrivatePCH.h"
 #include "IPlatformFileSandboxWrapper.h"
 #include "IPlatformFileModule.h"
+#include "UniquePtr.h"
 
 DEFINE_LOG_CATEGORY(SandboxFile);
 
@@ -356,8 +357,8 @@ class FSandboxFileModule : public IPlatformFileModule
 public:
 	virtual IPlatformFile* GetPlatformFile() override
 	{
-		static TScopedPointer<IPlatformFile> AutoDestroySingleton(new FSandboxPlatformFile(true));
-		return AutoDestroySingleton.GetOwnedPointer();
+		static TUniquePtr<IPlatformFile> AutoDestroySingleton = MakeUnique<FSandboxPlatformFile>(true);
+		return AutoDestroySingleton.Get();
 	}
 };
 IMPLEMENT_MODULE(FSandboxFileModule, SandboxFile);
diff --git a/Engine/Source/Runtime/SandboxFile/Public/IPlatformFileSandboxWrapper.h b/Engine/Source/Runtime/SandboxFile/Public/IPlatformFileSandboxWrapper.h
--- a/Engine/Source/Runtime/SandboxFile/Public/IPlatformFileSandboxWrapper.h
+++ b/Engine/Source/Runtime/SandboxFile/Public/IPlatformFileSandboxWrapper.h
@@ -3,6 +3,8 @@
 
 #pragma once
 
+#include "UniquePtr.h"
+
 /**
  * Wrapper to log the low level file system
 **/
@@ -11,8 +13,8 @@ DECLARE_LOG_CATEGORY_EXTERN(SandboxFile, Log, All);
 
 class SANDBOXFILE_API FSandboxFileHandle : public IFileHandle
 {
-	TAutoPtr<IFileHandle>	FileHandle;
-	FString				Filename;
+	TUniquePtr<IFileHandle>	FileHandle;
+	FString					Filename;
 
 public:
 
diff --git a/Engine/Source/Runtime/Serialization/Private/Backends/JsonStructDeserializerBackend.cpp b/Engine/Source/Runtime/Serialization/Private/Backends/JsonStructDeserializerBackend.cpp
--- a/Engine/Source/Runtime/Serialization/Private/Backends/JsonStructDeserializerBackend.cpp
+++ b/Engine/Source/Runtime/Serialization/Private/Backends/JsonStructDeserializerBackend.cpp
@@ -41,53 +41,60 @@ namespace JsonStructDeserializerBackend
 
 
 	/**
-	 * Sets the value of the given property.
+	 * Gets a pointer to object of the given property.
 	 *
-	 * @param Property The property to set.
-	 * @param Outer The property that contains the property to be set, if any.
+	 * @param Property The property to get.
+	 * @param Outer The property that contains the property to be get, if any.
 	 * @param Data A pointer to the memory holding the property's data.
 	 * @param ArrayIndex The index of the element to set (if the property is an array).
-	 * @return true on success, false otherwise.
+	 * @return A pointer to the object represented by the property, null otherwise..
 	 * @see ClearPropertyValue
 	 */
-	template<typename UPropertyType, typename PropertyType>
-	bool SetPropertyValue( UProperty* Property, UProperty* Outer, void* Data, int32 ArrayIndex, const PropertyType& Value )
+	void* GetPropertyValuePtr( UProperty* Property, UProperty* Outer, void* Data, int32 ArrayIndex )
 	{
-		PropertyType* ValuePtr = nullptr;
-		UArrayProperty* ArrayProperty = Cast<UArrayProperty>(Outer);
+		check(Property);
 
-		if (ArrayProperty != nullptr)
+		if (UArrayProperty* ArrayProperty = Cast<UArrayProperty>(Outer))
 		{
 			if (ArrayProperty->Inner != Property)
 			{
-				return false;
+				return nullptr;
 			}
 
 			FScriptArrayHelper ArrayHelper(ArrayProperty, ArrayProperty->template ContainerPtrToValuePtr<void>(Data));
 			int32 Index = ArrayHelper.AddValue();
 
-			ValuePtr = (PropertyType*)ArrayHelper.GetRawPtr(Index);
+			return ArrayHelper.GetRawPtr(Index);
 		}
-		else
-		{
-			UPropertyType* TypedProperty = Cast<UPropertyType>(Property);
-
-			if (TypedProperty == nullptr)
-			{
-				return false;
-			}
 
-			ValuePtr = TypedProperty->template ContainerPtrToValuePtr<PropertyType>(Data, ArrayIndex);
+		if (ArrayIndex >= Property->ArrayDim)
+		{
+			return nullptr;
 		}
 
-		if (ValuePtr == nullptr)
+		return Property->template ContainerPtrToValuePtr<void>(Data, ArrayIndex);
+	}
+
+	/**
+	 * Sets the value of the given property.
+	 *
+	 * @param Property The property to set.
+	 * @param Outer The property that contains the property to be set, if any.
+	 * @param Data A pointer to the memory holding the property's data.
+	 * @param ArrayIndex The index of the element to set (if the property is an array).
+	 * @return true on success, false otherwise.
+	 * @see ClearPropertyValue
+	 */
+	template<typename PropertyType, typename ValueType>
+	bool SetPropertyValue( PropertyType* Property, UProperty* Outer, void* Data, int32 ArrayIndex, const ValueType& Value )
+	{
+		if (void* Ptr = GetPropertyValuePtr(Property, Outer, Data, ArrayIndex))
 		{
-			return false;
+			*(ValueType*)Ptr = Value;
+			return true;
 		}
 
-		*ValuePtr = Value;
-
-		return true;
+		return false;
 	}
 }
 
@@ -170,9 +177,9 @@ bool FJsonStructDeserializerBackend::ReadProperty( UProperty* Property, UPropert
 		{
 			bool BoolValue = JsonReader->GetValueAsBoolean();
 
-			if (Property->GetClass() == UBoolProperty::StaticClass())
+			if (UBoolProperty* BoolProperty = Cast<UBoolProperty>(Property))
 			{
-				return SetPropertyValue<UBoolProperty, bool>(Property, Outer, Data, ArrayIndex, BoolValue);
+				return SetPropertyValue(BoolProperty, Outer, Data, ArrayIndex, BoolValue);
 			}
 
 			UE_LOG(LogSerialization, Verbose, TEXT("Boolean field %s with value '%s' is not supported in UProperty type %s (%s)"), *Property->GetFName().ToString(), BoolValue ? *(GTrue.ToString()) : *(GFalse.ToString()), *Property->GetClass()->GetName(), *GetDebugString());
@@ -186,54 +193,54 @@ bool FJsonStructDeserializerBackend::ReadProperty( UProperty* Property, UPropert
 		{
 			double NumericValue = JsonReader->GetValueAsNumber();
 
-			if (Property->GetClass() == UByteProperty::StaticClass())
+			if (UByteProperty* ByteProperty = Cast<UByteProperty>(Property))
 			{
-				return SetPropertyValue<UByteProperty, int8>(Property, Outer, Data, ArrayIndex, (int8)NumericValue);
+				return SetPropertyValue(ByteProperty, Outer, Data, ArrayIndex, (int8)NumericValue);
 			}
 
-			if (Property->GetClass() == UDoubleProperty::StaticClass())
+			if (UDoubleProperty* DoubleProperty = Cast<UDoubleProperty>(Property))
 			{
-				return SetPropertyValue<UDoubleProperty, double>(Property, Outer, Data, ArrayIndex, NumericValue);
+				return SetPropertyValue(DoubleProperty, Outer, Data, ArrayIndex, (double)NumericValue);
 			}
 
-			if (Property->GetClass() == UFloatProperty::StaticClass())
+			if (UFloatProperty* FloatProperty = Cast<UFloatProperty>(Property))
 			{
-				return SetPropertyValue<UFloatProperty, float>(Property, Outer, Data, ArrayIndex, NumericValue);
+				return SetPropertyValue(FloatProperty, Outer, Data, ArrayIndex, (float)NumericValue);
 			}
 
-			if (Property->GetClass() == UIntProperty::StaticClass())
+			if (UIntProperty* IntProperty = Cast<UIntProperty>(Property))
 			{
-				return SetPropertyValue<UIntProperty, int32>(Property, Outer, Data, ArrayIndex, (int32)NumericValue);
+				return SetPropertyValue(IntProperty, Outer, Data, ArrayIndex, (int32)NumericValue);
 			}
 
-			if (Property->GetClass() == UUInt32Property::StaticClass())
+			if (UUInt32Property* UInt32Property = Cast<UUInt32Property>(Property))
 			{
-				return SetPropertyValue<UUInt32Property, uint32>(Property, Outer, Data, ArrayIndex, (uint32)NumericValue);
+				return SetPropertyValue(UInt32Property, Outer, Data, ArrayIndex, (uint32)NumericValue);
 			}
 
-			if (Property->GetClass() == UInt16Property::StaticClass())
+			if (UInt16Property* Int16Property = Cast<UInt16Property>(Property))
 			{
-				return SetPropertyValue<UInt16Property, int16>(Property, Outer, Data, ArrayIndex, (int16)NumericValue);
+				return SetPropertyValue(Int16Property, Outer, Data, ArrayIndex, (int16)NumericValue);
 			}
 
-			if (Property->GetClass() == UUInt16Property::StaticClass())
+			if (UUInt16Property* UInt16Property = Cast<UUInt16Property>(Property))
 			{
-				return SetPropertyValue<UUInt16Property, uint16>(Property, Outer, Data, ArrayIndex, (uint16)NumericValue);
+				return SetPropertyValue(UInt16Property, Outer, Data, ArrayIndex, (uint16)NumericValue);
 			}
 
-			if (Property->GetClass() == UInt64Property::StaticClass())
+			if (UInt64Property* Int64Property = Cast<UInt64Property>(Property))
 			{
-				return SetPropertyValue<UInt64Property, int64>(Property, Outer, Data, ArrayIndex, (int64)NumericValue);
+				return SetPropertyValue(Int64Property, Outer, Data, ArrayIndex, (int64)NumericValue);
 			}
 
-			if (Property->GetClass() == UUInt64Property::StaticClass())
+			if (UUInt64Property* UInt64Property = Cast<UUInt64Property>(Property))
 			{
-				return SetPropertyValue<UUInt64Property, uint64>(Property, Outer, Data, ArrayIndex, (uint64)NumericValue);
+				return SetPropertyValue(UInt64Property, Outer, Data, ArrayIndex, (uint64)NumericValue);
 			}
 
-			if (Property->GetClass() == UInt8Property::StaticClass())
+			if (UInt8Property* Int8Property = Cast<UInt8Property>(Property))
 			{
-				return SetPropertyValue<UInt8Property, int8>(Property, Outer, Data, ArrayIndex, (int8)NumericValue);
+				return SetPropertyValue(Int8Property, Outer, Data, ArrayIndex, (int8)NumericValue);
 			}
 
 			UE_LOG(LogSerialization, Verbose, TEXT("Numeric field %s with value '%f' is not supported in UProperty type %s (%s)"), *Property->GetFName().ToString(), NumericValue, *Property->GetClass()->GetName(), *GetDebugString());
@@ -251,37 +258,57 @@ bool FJsonStructDeserializerBackend::ReadProperty( UProperty* Property, UPropert
 		{
 			const FString& StringValue = JsonReader->GetValueAsString();
 
-			if (Property->GetClass() == UStrProperty::StaticClass())
+			if (UStrProperty* StrProperty = Cast<UStrProperty>(Property))
 			{
-				return SetPropertyValue<UStrProperty, FString>(Property, Outer, Data, ArrayIndex, StringValue);
+				return SetPropertyValue(StrProperty, Outer, Data, ArrayIndex, StringValue);
 			}
 
-			if (Property->GetClass() == UNameProperty::StaticClass())
+			if (UNameProperty* NameProperty = Cast<UNameProperty>(Property))
 			{
-				return SetPropertyValue<UNameProperty, FName>(Property, Outer, Data, ArrayIndex, *StringValue);
+				return SetPropertyValue(NameProperty, Outer, Data, ArrayIndex, FName(*StringValue));
 			}
 
-			if (Property->GetClass() == UTextProperty::StaticClass())
+			if (UTextProperty* TextProperty = Cast<UTextProperty>(Property))
 			{
-				return SetPropertyValue<UTextProperty, FText>(Property, Outer, Data, ArrayIndex, FText::FromString(StringValue));
+				return SetPropertyValue(TextProperty, Outer, Data, ArrayIndex, FText::FromString(StringValue));
 			}
 
-			if (Property->GetClass() == UByteProperty::StaticClass())
+			if (UByteProperty* ByteProperty = Cast<UByteProperty>(Property))
 			{
-				UByteProperty* ByteProperty = Cast<UByteProperty>(Property);
+				if (!ByteProperty->Enum)
+				{
+					return false;
+				}
+
 				int32 Value = ByteProperty->Enum->GetValueByName(*StringValue);
+				if (Value == INDEX_NONE)
+				{
+					return false;
+				}
+
+				return SetPropertyValue(ByteProperty, Outer, Data, ArrayIndex, (uint8)Value);
+			}
 
+			if (UEnumProperty* EnumProperty = Cast<UEnumProperty>(Property))
+			{
+				int64 Value = EnumProperty->GetEnum()->GetValueByName(*StringValue);
 				if (Value == INDEX_NONE)
 				{
 					return false;
 				}
 
-				return SetPropertyValue<UByteProperty, uint8>(Property, Outer, Data, ArrayIndex, (uint8)Value);
+				if (void* ElementPtr = GetPropertyValuePtr(EnumProperty, Outer, Data, ArrayIndex))
+				{
+					EnumProperty->GetUnderlyingProperty()->SetIntPropertyValue(ElementPtr, Value);
+					return true;
+				}
+
+				return false;
 			}
 
-			if (Property->GetClass() == UClassProperty::StaticClass())
+			if (UClassProperty* ClassProperty = Cast<UClassProperty>(Property))
 			{
-				return SetPropertyValue<UClassProperty, UClass*>(Property, Outer, Data, ArrayIndex, LoadObject<UClass>(NULL, *StringValue, NULL, LOAD_NoWarn));
+				return SetPropertyValue(ClassProperty, Outer, Data, ArrayIndex, LoadObject<UClass>(NULL, *StringValue, NULL, LOAD_NoWarn));
 			}
 
 			UE_LOG(LogSerialization, Verbose, TEXT("String field %s with value '%s' is not supported in UProperty type %s (%s)"), *Property->GetFName().ToString(), *StringValue, *Property->GetClass()->GetName(), *GetDebugString());
diff --git a/Engine/Source/Runtime/Serialization/Private/Backends/JsonStructSerializerBackend.cpp b/Engine/Source/Runtime/Serialization/Private/Backends/JsonStructSerializerBackend.cpp
--- a/Engine/Source/Runtime/Serialization/Private/Backends/JsonStructSerializerBackend.cpp
+++ b/Engine/Source/Runtime/Serialization/Private/Backends/JsonStructSerializerBackend.cpp
@@ -131,6 +131,12 @@ void FJsonStructSerializerBackend::WriteProperty(const FStructSerializerState& S
 	}
 
 	// unsigned bytes & enumerations
+	else if (State.ValueType == UEnumProperty::StaticClass())
+	{
+		UEnumProperty* EnumProperty = Cast<UEnumProperty>(State.ValueProperty);
+
+		WritePropertyValue(JsonWriter, State, EnumProperty->GetEnum()->GetEnumNameStringByValue(EnumProperty->GetUnderlyingProperty()->GetSignedIntPropertyValue(EnumProperty->ContainerPtrToValuePtr<void>(State.ValueData, ArrayIndex))));
+	}
 	else if (State.ValueType == UByteProperty::StaticClass())
 	{
 		UByteProperty* ByteProperty = Cast<UByteProperty>(State.ValueProperty);
diff --git a/Engine/Source/Runtime/Sockets/Private/SocketSubsystem.cpp b/Engine/Source/Runtime/Sockets/Private/SocketSubsystem.cpp
--- a/Engine/Source/Runtime/Sockets/Private/SocketSubsystem.cpp
+++ b/Engine/Source/Runtime/Sockets/Private/SocketSubsystem.cpp
@@ -4,6 +4,7 @@
 #include "SocketSubsystem.h"
 #include "ModuleManager.h"
 #include "Sockets.h"
+#include "UniquePtr.h"
 
 DEFINE_LOG_CATEGORY(LogSockets);
 
@@ -194,8 +195,8 @@ ISocketSubsystem* FSocketSubsystemModule::GetSocketSubsystem(const FName InSubsy
 ISocketSubsystem* ISocketSubsystem::Get(const FName& SubsystemName)
 {
 	// wrap the platform file with a logger
-	//	static TScopedPointer<FLoggedPlatformFile> AutoDestroyLog;
-	//	AutoDestroyLog = new FLoggedPlatformFile(*CurrentPlatformFile);
+	//	static TUniquePtr<FLoggedPlatformFile> AutoDestroyLog;
+	//	AutoDestroyLog = MakeUnique<FLoggedPlatformFile>(*CurrentPlatformFile);
 
 	struct FStatic
 	{
diff --git a/Engine/Source/Runtime/StreamingFile/Private/StreamingNetworkPlatformFile.cpp b/Engine/Source/Runtime/StreamingFile/Private/StreamingNetworkPlatformFile.cpp
--- a/Engine/Source/Runtime/StreamingFile/Private/StreamingNetworkPlatformFile.cpp
+++ b/Engine/Source/Runtime/StreamingFile/Private/StreamingNetworkPlatformFile.cpp
@@ -7,7 +7,7 @@
 #include "FileManagerGeneric.h"
 #include "MultichannelTCP.h"
 #include "IPlatformFileModule.h"
-
+#include "UniquePtr.h"
 
 DEFINE_LOG_CATEGORY(LogStreamingPlatformFile);
 
@@ -941,8 +941,8 @@ public:
 
 	virtual IPlatformFile* GetPlatformFile() override
 	{
-		static TScopedPointer<IPlatformFile> AutoDestroySingleton(new FStreamingNetworkPlatformFile());
-		return AutoDestroySingleton.GetOwnedPointer();
+		static TUniquePtr<IPlatformFile> AutoDestroySingleton = MakeUnique<FStreamingNetworkPlatformFile>();
+		return AutoDestroySingleton.Get();
 	}
 };
 
diff --git a/Engine/Source/Runtime/UMG/Private/Binding/CheckedStateBinding.cpp b/Engine/Source/Runtime/UMG/Private/Binding/CheckedStateBinding.cpp
--- a/Engine/Source/Runtime/UMG/Private/Binding/CheckedStateBinding.cpp
+++ b/Engine/Source/Runtime/UMG/Private/Binding/CheckedStateBinding.cpp
@@ -18,7 +18,11 @@ bool UCheckedStateBinding::IsSupportedDestination(UProperty* Property) const
 {
 	static const FName CheckBoxStateEnum(TEXT("ECheckBoxState"));
 
-	if ( UByteProperty* ByteProperty = Cast<UByteProperty>(Property) )
+	if ( UEnumProperty* EnumProperty = Cast<UEnumProperty>(Property) )
+	{
+		return EnumProperty->GetEnum()->GetFName() == CheckBoxStateEnum;
+	}
+	else if ( UByteProperty* ByteProperty = Cast<UByteProperty>(Property) )
 	{
 		if ( ByteProperty->IsEnum() )
 		{
diff --git a/Engine/Source/Runtime/UMG/Private/Binding/MouseCursorBinding.cpp b/Engine/Source/Runtime/UMG/Private/Binding/MouseCursorBinding.cpp
--- a/Engine/Source/Runtime/UMG/Private/Binding/MouseCursorBinding.cpp
+++ b/Engine/Source/Runtime/UMG/Private/Binding/MouseCursorBinding.cpp
@@ -18,7 +18,11 @@ bool UMouseCursorBinding::IsSupportedDestination(UProperty* Property) const
 {
 	static const FName MouseCursorEnum(TEXT("EMouseCursor"));
 	
-	if ( UByteProperty* ByteProperty = Cast<UByteProperty>(Property) )
+	if ( UEnumProperty* EnumProperty = Cast<UEnumProperty>(Property) )
+	{
+		return EnumProperty->GetEnum()->GetFName() == MouseCursorEnum;
+	}
+	else if ( UByteProperty* ByteProperty = Cast<UByteProperty>(Property) )
 	{
 		if ( ByteProperty->IsEnum() )
 		{
diff --git a/Engine/Source/Runtime/UMG/Private/Binding/VisibilityBinding.cpp b/Engine/Source/Runtime/UMG/Private/Binding/VisibilityBinding.cpp
--- a/Engine/Source/Runtime/UMG/Private/Binding/VisibilityBinding.cpp
+++ b/Engine/Source/Runtime/UMG/Private/Binding/VisibilityBinding.cpp
@@ -18,7 +18,11 @@ bool UVisibilityBinding::IsSupportedDestination(UProperty* Property) const
 {
 	static const FName VisibilityEnum(TEXT("ESlateVisibility"));
 
-	if ( UByteProperty* ByteProperty = Cast<UByteProperty>(Property) )
+	if ( UEnumProperty* EnumProperty = Cast<UEnumProperty>(Property) )
+	{
+		return EnumProperty->GetEnum()->GetFName() == VisibilityEnum;
+	}
+	else if ( UByteProperty* ByteProperty = Cast<UByteProperty>(Property) )
 	{
 		if ( ByteProperty->IsEnum() )
 		{
diff --git a/Engine/Source/Runtime/UMG/Private/Components/PanelWidget.cpp b/Engine/Source/Runtime/UMG/Private/Components/PanelWidget.cpp
--- a/Engine/Source/Runtime/UMG/Private/Components/PanelWidget.cpp
+++ b/Engine/Source/Runtime/UMG/Private/Components/PanelWidget.cpp
@@ -114,8 +114,13 @@ UPanelSlot* UPanelWidget::AddChild(UWidget* Content)
 
 	Content->RemoveFromParent();
 
-	UPanelSlot* PanelSlot = NewObject<UPanelSlot>(this, GetSlotClass());
-	PanelSlot->SetFlags(RF_Transactional);
+	EObjectFlags NewObjectFlags = RF_Transactional;
+	if (HasAnyFlags(RF_Transient))
+	{
+		NewObjectFlags |= RF_Transient;
+	}
+
+	UPanelSlot* PanelSlot = NewObject<UPanelSlot>(this, GetSlotClass(), NAME_None, NewObjectFlags);
 	PanelSlot->Content = Content;
 	PanelSlot->Parent = this;
 
diff --git a/Engine/Source/Runtime/UMG/Private/UserWidget.cpp b/Engine/Source/Runtime/UMG/Private/UserWidget.cpp
--- a/Engine/Source/Runtime/UMG/Private/UserWidget.cpp
+++ b/Engine/Source/Runtime/UMG/Private/UserWidget.cpp
@@ -84,7 +84,7 @@ bool UUserWidget::Initialize()
 
 		if ( WidgetTree == nullptr )
 		{
-			WidgetTree = NewObject<UWidgetTree>(this, TEXT("WidgetTree"));
+			WidgetTree = NewObject<UWidgetTree>(this, TEXT("WidgetTree"), RF_Transient);
 		}
 
 		// Map the named slot bindings to the available slots.
@@ -264,7 +264,7 @@ UUMGSequencePlayer* UUserWidget::GetOrAddPlayer(UWidgetAnimation* InAnimation)
 
 		if (!FoundPlayer)
 		{
-			UUMGSequencePlayer* NewPlayer = NewObject<UUMGSequencePlayer>(this);
+			UUMGSequencePlayer* NewPlayer = NewObject<UUMGSequencePlayer>(this, NAME_None, RF_Transient);
 			ActiveSequencePlayers.Add(NewPlayer);
 
 			NewPlayer->OnSequenceFinishedPlaying().AddUObject(this, &UUserWidget::OnAnimationFinishedPlaying);
diff --git a/Engine/Source/Runtime/UMG/Private/WidgetBlueprintGeneratedClass.cpp b/Engine/Source/Runtime/UMG/Private/WidgetBlueprintGeneratedClass.cpp
--- a/Engine/Source/Runtime/UMG/Private/WidgetBlueprintGeneratedClass.cpp
+++ b/Engine/Source/Runtime/UMG/Private/WidgetBlueprintGeneratedClass.cpp
@@ -52,9 +52,14 @@ void UWidgetBlueprintGeneratedClass::InitializeWidgetStatic(UUserWidget* UserWid
 	// similar to how we use to use the DesignerWidgetTree.
 	if ( ClonedTree == nullptr )
 	{
-		ClonedTree = DuplicateObject<UWidgetTree>(InWidgetTree, UserWidget);
-	}
+		FObjectDuplicationParameters Parameters(InWidgetTree, UserWidget);
+
+		// Set to be transient and strip public flags
+		Parameters.ApplyFlags = RF_Transient;
+		Parameters.FlagMask = Parameters.FlagMask & ~(RF_Public | RF_DefaultSubObject);
 
+		ClonedTree = Cast<UWidgetTree>(StaticDuplicateObjectEx(Parameters));
+	}
 	UserWidget->WidgetGeneratedByClass = InClass;
 
 #if WITH_EDITOR
@@ -172,6 +177,12 @@ void UWidgetBlueprintGeneratedClass::PostLoad()
 {
 	Super::PostLoad();
 
+	// Clear CDO flag on tree
+	if (WidgetTree)
+	{
+		WidgetTree->ClearFlags(RF_DefaultSubObject);
+	}
+
 	if ( GetLinkerUE4Version() < VER_UE4_RENAME_WIDGET_VISIBILITY )
 	{
 		static const FName Visiblity(TEXT("Visiblity"));
diff --git a/Engine/Source/Runtime/UMG/Public/Binding/DynamicPropertyPath.h b/Engine/Source/Runtime/UMG/Public/Binding/DynamicPropertyPath.h
--- a/Engine/Source/Runtime/UMG/Public/Binding/DynamicPropertyPath.h
+++ b/Engine/Source/Runtime/UMG/Public/Binding/DynamicPropertyPath.h
@@ -25,48 +25,88 @@ inline bool IsConcreteTypeCompatibleWithReflectedType<bool>(UProperty* Property)
 template<>
 inline bool IsConcreteTypeCompatibleWithReflectedType<int8>(UProperty* Property)
 {
+	if (UEnumProperty* EnumProperty = Cast<UEnumProperty>(Property))
+	{
+		Property = EnumProperty->GetUnderlyingProperty();
+	}
+
 	return Property->GetClass() == UInt8Property::StaticClass();
 }
 
 template<>
 inline bool IsConcreteTypeCompatibleWithReflectedType<uint8>(UProperty* Property)
 {
+	if (UEnumProperty* EnumProperty = Cast<UEnumProperty>(Property))
+	{
+		Property = EnumProperty->GetUnderlyingProperty();
+	}
+
 	return Property->GetClass() == UByteProperty::StaticClass();
 }
 
 template<>
 inline bool IsConcreteTypeCompatibleWithReflectedType<int16>(UProperty* Property)
 {
+	if (UEnumProperty* EnumProperty = Cast<UEnumProperty>(Property))
+	{
+		Property = EnumProperty->GetUnderlyingProperty();
+	}
+
 	return Property->GetClass() == UInt16Property::StaticClass();
 }
 
 template<>
 inline bool IsConcreteTypeCompatibleWithReflectedType<uint16>(UProperty* Property)
 {
+	if (UEnumProperty* EnumProperty = Cast<UEnumProperty>(Property))
+	{
+		Property = EnumProperty->GetUnderlyingProperty();
+	}
+
 	return Property->GetClass() == UUInt16Property::StaticClass();
 }
 
 template<>
 inline bool IsConcreteTypeCompatibleWithReflectedType<int32>(UProperty* Property)
 {
+	if (UEnumProperty* EnumProperty = Cast<UEnumProperty>(Property))
+	{
+		Property = EnumProperty->GetUnderlyingProperty();
+	}
+
 	return Property->GetClass() == UIntProperty::StaticClass();
 }
 
 template<>
 inline bool IsConcreteTypeCompatibleWithReflectedType<uint32>(UProperty* Property)
 {
+	if (UEnumProperty* EnumProperty = Cast<UEnumProperty>(Property))
+	{
+		Property = EnumProperty->GetUnderlyingProperty();
+	}
+
 	return Property->GetClass() == UUInt32Property::StaticClass();
 }
 
 template<>
 inline bool IsConcreteTypeCompatibleWithReflectedType<int64>(UProperty* Property)
 {
+	if (UEnumProperty* EnumProperty = Cast<UEnumProperty>(Property))
+	{
+		Property = EnumProperty->GetUnderlyingProperty();
+	}
+
 	return Property->GetClass() == UInt64Property::StaticClass();
 }
 
 template<>
 inline bool IsConcreteTypeCompatibleWithReflectedType<uint64>(UProperty* Property)
 {
+	if (UEnumProperty* EnumProperty = Cast<UEnumProperty>(Property))
+	{
+		Property = EnumProperty->GetUnderlyingProperty();
+	}
+
 	return Property->GetClass() == UUInt64Property::StaticClass();
 }
 
diff --git a/Engine/Source/Runtime/UMG/Public/Blueprint/WidgetTree.h b/Engine/Source/Runtime/UMG/Public/Blueprint/WidgetTree.h
--- a/Engine/Source/Runtime/UMG/Public/Blueprint/WidgetTree.h
+++ b/Engine/Source/Runtime/UMG/Public/Blueprint/WidgetTree.h
@@ -107,17 +107,21 @@ public:
 	template< class T >
 	FORCEINLINE T* ConstructWidget(TSubclassOf<UWidget> WidgetType = T::StaticClass(), FName WidgetName = NAME_None)
 	{
+		EObjectFlags NewObjectFlags = RF_Transactional;
+		if (HasAnyFlags(RF_Transient))
+		{
+			NewObjectFlags |= RF_Transient;
+		}
+
 		if ( WidgetType->IsChildOf(UUserWidget::StaticClass()) )
 		{
-			UUserWidget* Widget = NewObject<UUserWidget>(this, WidgetType, WidgetName);
+			UUserWidget* Widget = NewObject<UUserWidget>(this, WidgetType, WidgetName, NewObjectFlags);
 			Widget->Initialize();
-			Widget->SetFlags(RF_Transactional);
 			return (T*)Widget;
 		}
 		else
 		{
-			UWidget* Widget = (UWidget*)NewObject<UWidget>(this, WidgetType, WidgetName);
-			Widget->SetFlags(RF_Transactional);
+			UWidget* Widget = (UWidget*)NewObject<UWidget>(this, WidgetType, WidgetName, NewObjectFlags);
 			return (T*)Widget;
 		}
 	}
diff --git a/Engine/Source/Runtime/WebBrowser/Private/CEF/CEFJSStructDeserializerBackend.cpp b/Engine/Source/Runtime/WebBrowser/Private/CEF/CEFJSStructDeserializerBackend.cpp
--- a/Engine/Source/Runtime/WebBrowser/Private/CEF/CEFJSStructDeserializerBackend.cpp
+++ b/Engine/Source/Runtime/WebBrowser/Private/CEF/CEFJSStructDeserializerBackend.cpp
@@ -70,63 +70,69 @@ namespace {
 	}
 
 	/**
-	 * Sets the value of the given property.
+	 * Gets a pointer to object of the given property.
 	 *
-	 * @param Property The property to set.
-	 * @param Outer The property that contains the property to be set, if any.
+	 * @param Property The property to get.
+	 * @param Outer The property that contains the property to be get, if any.
 	 * @param Data A pointer to the memory holding the property's data.
 	 * @param ArrayIndex The index of the element to set (if the property is an array).
-	 * @return true on success, false otherwise.
+	 * @return A pointer to the object represented by the property, null otherwise..
 	 * @see ClearPropertyValue
 	 */
-	template<typename UPropertyType, typename PropertyType>
-	bool SetPropertyValue( UProperty* Property, UProperty* Outer, void* Data, int32 ArrayIndex, const PropertyType& Value )
+	void* GetPropertyValuePtr( UProperty* Property, UProperty* Outer, void* Data, int32 ArrayIndex )
 	{
-		PropertyType* ValuePtr = nullptr;
-		UArrayProperty* ArrayProperty = Cast<UArrayProperty>(Outer);
+		check(Property);
 
-		if (ArrayProperty != nullptr)
+		if (UArrayProperty* ArrayProperty = Cast<UArrayProperty>(Outer))
 		{
 			if (ArrayProperty->Inner != Property)
 			{
-				return false;
+				return nullptr;
 			}
 
 			FScriptArrayHelper ArrayHelper(ArrayProperty, ArrayProperty->template ContainerPtrToValuePtr<void>(Data));
 			int32 Index = ArrayHelper.AddValue();
 
-			ValuePtr = (PropertyType*)ArrayHelper.GetRawPtr(Index);
+			return ArrayHelper.GetRawPtr(Index);
 		}
-		else
-		{
-			UPropertyType* TypedProperty = Cast<UPropertyType>(Property);
-
-			if (TypedProperty == nullptr || ArrayIndex >= TypedProperty->ArrayDim)
-			{
-				return false;
-			}
 
-			ValuePtr = TypedProperty->template ContainerPtrToValuePtr<PropertyType>(Data, ArrayIndex);
+		if (ArrayIndex >= Property->ArrayDim)
+		{
+			return nullptr;
 		}
 
-		if (ValuePtr == nullptr)
+		return Property->template ContainerPtrToValuePtr<void>(Data, ArrayIndex);
+	}
+
+	/**
+	 * Sets the value of the given property.
+	 *
+	 * @param Property The property to set.
+	 * @param Outer The property that contains the property to be set, if any.
+	 * @param Data A pointer to the memory holding the property's data.
+	 * @param ArrayIndex The index of the element to set (if the property is an array).
+	 * @return true on success, false otherwise.
+	 * @see ClearPropertyValue
+	 */
+	template<typename PropertyType, typename ValueType>
+	bool SetPropertyValue( PropertyType* Property, UProperty* Outer, void* Data, int32 ArrayIndex, const ValueType& Value )
+	{
+		if (void* Ptr = GetPropertyValuePtr(Property, Outer, Data, ArrayIndex))
 		{
-			return false;
+			*(ValueType*)Ptr = Value;
+			return true;
 		}
 
-		*ValuePtr = Value;
-
-		return true;
+		return false;
 	}
 
-
 	template<typename PropertyType, typename ContainerType, typename KeyType>
 	bool ReadNumericProperty(UProperty* Property, UProperty* Outer, void* Data, int32 ArrayIndex, CefRefPtr<ContainerType> Container, KeyType Key )
 	{
 		typedef typename PropertyType::TCppType TCppType;
-		if (Property->IsA<PropertyType>())
+		if (PropertyType* TypedProperty = Cast<PropertyType>(Property))
 		{
-			return SetPropertyValue<PropertyType, TCppType>(Property, Outer, Data, ArrayIndex, GetNumeric<TCppType>(Container, Key));
+			return SetPropertyValue(TypedProperty, Outer, Data, ArrayIndex, GetNumeric<TCppType>(Container, Key));
 		}
 		else
 		{
@@ -137,9 +143,9 @@ namespace {
 	template<typename ContainerType, typename KeyType>
 	bool ReadBoolProperty(UProperty* Property, UProperty* Outer, void* Data, int32 ArrayIndex, CefRefPtr<ContainerType> Container, KeyType Key )
 	{
-		if (Property->IsA<UBoolProperty>())
+		if (UBoolProperty* BoolProperty = Cast<UBoolProperty>(Property))
 		{
-			return SetPropertyValue<UBoolProperty, bool>(Property, Outer, Data, ArrayIndex, GetNumeric<int>(Container, Key)!=0);
+			return SetPropertyValue(BoolProperty, Outer, Data, ArrayIndex, GetNumeric<int>(Container, Key)!=0);
 		}
 		return false;
 
@@ -155,7 +161,7 @@ namespace {
 		CefRefPtr<CefDictionaryValue> Dictionary = Container->GetDictionary(Key);
 		UStructProperty* StructProperty = Cast<UStructProperty>(Property);
 
-		if ( StructProperty->Struct != FWebJSFunction::StaticStruct())
+		if ( !StructProperty || StructProperty->Struct != FWebJSFunction::StaticStruct())
 		{
 			return false;
 		}
@@ -167,7 +173,7 @@ namespace {
 			return false;
 		}
 		FWebJSFunction CallbackObject(Scripting, CallbackID);
-		return SetPropertyValue<UStructProperty, FWebJSFunction>(Property, Outer, Data, ArrayIndex, CallbackObject);
+		return SetPropertyValue(StructProperty, Outer, Data, ArrayIndex, CallbackObject);
 	}
 
 	template<typename ContainerType, typename KeyType>
@@ -177,29 +183,52 @@ namespace {
 		{
 			FString StringValue = Container->GetString(Key).ToWString().c_str();
 
-			if (Property->IsA<UStrProperty>())
+			if (UStrProperty* StrProperty = Cast<UStrProperty>(Property))
 			{
-				return SetPropertyValue<UStrProperty, FString>(Property, Outer, Data, ArrayIndex, StringValue);
+				return SetPropertyValue(StrProperty, Outer, Data, ArrayIndex, StringValue);
 			}
-			else if (Property->IsA<UNameProperty>())
+
+			if (UNameProperty* NameProperty = Cast<UNameProperty>(Property))
 			{
-				return SetPropertyValue<UNameProperty, FName>(Property, Outer, Data, ArrayIndex, *StringValue);
+				return SetPropertyValue(NameProperty, Outer, Data, ArrayIndex, FName(*StringValue));
 			}
-			else if (Property->IsA<UTextProperty>())
+
+			if (UTextProperty* TextProperty = Cast<UTextProperty>(Property))
 			{
-				return SetPropertyValue<UTextProperty, FText>(Property, Outer, Data, ArrayIndex, FText::FromString(StringValue));
+				return SetPropertyValue(TextProperty, Outer, Data, ArrayIndex, FText::FromString(StringValue));
 			}
-			else if (Property->IsA<UByteProperty>())
+
+			if (UByteProperty* ByteProperty = Cast<UByteProperty>(Property))
 			{
-				UByteProperty* ByteProperty = Cast<UByteProperty>(Property);
+				if (!ByteProperty->Enum)
+				{
+					return false;
+				}
+
 				int32 Index = ByteProperty->Enum->FindEnumIndex(*StringValue);
+				if (Index == INDEX_NONE)
+				{
+					return false;
+				}
+
+				return SetPropertyValue(ByteProperty, Outer, Data, ArrayIndex, (uint8)Index);
+			}
 
+			if (UEnumProperty* EnumProperty = Cast<UEnumProperty>(Property))
+			{
+				int32 Index = EnumProperty->GetEnum()->FindEnumIndex(*StringValue);
 				if (Index == INDEX_NONE)
 				{
 					return false;
 				}
 
-				return SetPropertyValue<UByteProperty, uint8>(Property, Outer, Data, ArrayIndex, (uint8)Index);
+				if (void* ElementPtr = GetPropertyValuePtr(EnumProperty, Outer, Data, ArrayIndex))
+				{
+					EnumProperty->GetUnderlyingProperty()->SetIntPropertyValue(ElementPtr, (int64)Index);
+					return true;
+				}
+
+				return false;
 			}
 		}
 
diff --git a/Engine/Source/Runtime/WebBrowser/Private/CEF/CEFJSStructSerializerBackend.cpp b/Engine/Source/Runtime/WebBrowser/Private/CEF/CEFJSStructSerializerBackend.cpp
--- a/Engine/Source/Runtime/WebBrowser/Private/CEF/CEFJSStructSerializerBackend.cpp
+++ b/Engine/Source/Runtime/WebBrowser/Private/CEF/CEFJSStructSerializerBackend.cpp
@@ -184,6 +184,12 @@ void FCEFJSStructSerializerBackend::WriteProperty(const FStructSerializerState&
 	}
 
 	// unsigned bytes & enumerations
+	else if (State.ValueType == UEnumProperty::StaticClass())
+	{
+		UEnumProperty* EnumProperty = Cast<UEnumProperty>(State.ValueProperty);
+
+		Add(State, EnumProperty->GetEnum()->GetEnumName(EnumProperty->GetUnderlyingProperty()->GetSignedIntPropertyValue(EnumProperty->ContainerPtrToValuePtr<void>(State.ValueData, ArrayIndex))));
+	}
 	else if (State.ValueType == UByteProperty::StaticClass())
 	{
 		UByteProperty* ByteProperty = Cast<UByteProperty>(State.ValueProperty);
@@ -194,7 +200,7 @@ void FCEFJSStructSerializerBackend::WriteProperty(const FStructSerializerState&
 		}
 		else
 		{
-			Add(State, (double)Cast<UByteProperty>(State.ValueProperty)->GetPropertyValue_InContainer(State.ValueData, ArrayIndex));
+			Add(State, (double)ByteProperty->GetPropertyValue_InContainer(State.ValueData, ArrayIndex));
 		}
 	}
 
diff --git a/Engine/Source/Runtime/Windows/D3D11RHI/Private/D3D11Device.cpp b/Engine/Source/Runtime/Windows/D3D11RHI/Private/D3D11Device.cpp
--- a/Engine/Source/Runtime/Windows/D3D11RHI/Private/D3D11Device.cpp
+++ b/Engine/Source/Runtime/Windows/D3D11RHI/Private/D3D11Device.cpp
@@ -188,6 +188,7 @@ FD3D11DynamicRHI::FD3D11DynamicRHI(IDXGIFactory1* InDXGIFactory1,D3D_FEATURE_LEV
 	GMaxShadowDepthBufferSizeX = 4096;
 	GMaxShadowDepthBufferSizeY = 4096;
 	GSupportsTimestampRenderQueries = true;
+	GRHISupportsResolveCubemapFaces = true;
 
 	// Initialize the constant buffers.
 	InitConstantBuffers();
