Fixed deserialization into bitfield property in StructDeserializer
Added proper bitfield serialization test for StructSerializer

diff --git a/Engine/Source/Runtime/Serialization/Private/Backends/StructDeserializerBackendUtilities.h b/Engine/Source/Runtime/Serialization/Private/Backends/StructDeserializerBackendUtilities.h
--- a/Engine/Source/Runtime/Serialization/Private/Backends/StructDeserializerBackendUtilities.h
+++ b/Engine/Source/Runtime/Serialization/Private/Backends/StructDeserializerBackendUtilities.h
@@ -92,7 +92,7 @@ struct StructDeserializerBackendUtilities
 	{
 		if (void* Ptr = GetPropertyValuePtr(Property, Outer, Data, ArrayIndex))
 		{
-			*(ValueType*)Ptr = Value;
+			Property->SetPropertyValue(Ptr, Value);
 			return true;
 		}
 
diff --git a/Engine/Source/Runtime/Serialization/Private/Tests/StructSerializerTest.cpp b/Engine/Source/Runtime/Serialization/Private/Tests/StructSerializerTest.cpp
--- a/Engine/Source/Runtime/Serialization/Private/Tests/StructSerializerTest.cpp
+++ b/Engine/Source/Runtime/Serialization/Private/Tests/StructSerializerTest.cpp
@@ -68,7 +68,14 @@ namespace StructSerializerTest
 		// test booleans
 		Test.TestEqual<bool>(TEXT("Booleans.BoolFalse must be the same before and after de-/serialization"), TestStruct.Booleans.BoolFalse, TestStruct2.Booleans.BoolFalse);
 		Test.TestEqual<bool>(TEXT("Booleans.BoolTrue must be the same before and after de-/serialization"), TestStruct.Booleans.BoolTrue, TestStruct2.Booleans.BoolTrue);
-		Test.TestEqual<uint32>(TEXT("Booleans.Bitfield must be the same before and after de-/serialization"), TestStruct.Booleans.Bitfield, TestStruct2.Booleans.Bitfield);
+		Test.TestEqual<bool>(TEXT("Booleans.Bitfield0 must be the same before and after de-/serialization"), TestStruct.Booleans.Bitfield0, TestStruct2.Booleans.Bitfield0);
+		Test.TestEqual<bool>(TEXT("Booleans.Bitfield1 must be the same before and after de-/serialization"), TestStruct.Booleans.Bitfield1, TestStruct2.Booleans.Bitfield1);
+		Test.TestEqual<bool>(TEXT("Booleans.Bitfield2Set must be the same before and after de-/serialization"), TestStruct.Booleans.Bitfield2Set, TestStruct2.Booleans.Bitfield2Set);
+		Test.TestEqual<bool>(TEXT("Booleans.Bitfield3 must be the same before and after de-/serialization"), TestStruct.Booleans.Bitfield3, TestStruct2.Booleans.Bitfield3);
+		Test.TestEqual<bool>(TEXT("Booleans.Bitfield4Set must be the same before and after de-/serialization"), TestStruct.Booleans.Bitfield4Set, TestStruct2.Booleans.Bitfield4Set);
+		Test.TestEqual<bool>(TEXT("Booleans.Bitfield5Set must be the same before and after de-/serialization"), TestStruct.Booleans.Bitfield5Set, TestStruct2.Booleans.Bitfield5Set);
+		Test.TestEqual<bool>(TEXT("Booleans.Bitfield6 must be the same before and after de-/serialization"), TestStruct.Booleans.Bitfield6, TestStruct2.Booleans.Bitfield6);
+		Test.TestEqual<bool>(TEXT("Booleans.Bitfield7 must be the same before and after de-/serialization"), TestStruct.Booleans.Bitfield7Set, TestStruct2.Booleans.Bitfield7Set);
 
 		// test objects
 		Test.TestEqual<class UClass*>(TEXT("Objects.Class must be the same before and after de-/serialization"), TestStruct.Objects.Class, TestStruct2.Objects.Class);
diff --git a/Engine/Source/Runtime/Serialization/Private/Tests/StructSerializerTestTypes.h b/Engine/Source/Runtime/Serialization/Private/Tests/StructSerializerTestTypes.h
--- a/Engine/Source/Runtime/Serialization/Private/Tests/StructSerializerTestTypes.h
+++ b/Engine/Source/Runtime/Serialization/Private/Tests/StructSerializerTestTypes.h
@@ -80,13 +80,41 @@ struct FStructSerializerBooleanTestStruct
 	bool BoolTrue;
 
 	UPROPERTY()
-	uint32 Bitfield;
+	uint8 Bitfield0 : 1;
+
+	UPROPERTY()
+	uint8 Bitfield1 : 1;
+
+	UPROPERTY()
+	uint8 Bitfield2Set : 1;
+
+	UPROPERTY()
+	uint8 Bitfield3 : 1;
+
+	UPROPERTY()
+	uint8 Bitfield4Set : 1;
+
+	UPROPERTY()
+	uint8 Bitfield5Set : 1;
+
+	UPROPERTY()
+	uint8 Bitfield6 : 1;
+
+	UPROPERTY()
+	uint8 Bitfield7Set : 1;
 
 	/** Default constructor. */
 	FStructSerializerBooleanTestStruct()
 		: BoolFalse(false)
 		, BoolTrue(true)
-		, Bitfield(1)
+		, Bitfield0(0)
+		, Bitfield1(0)
+		, Bitfield2Set(1)
+		, Bitfield3(0)
+		, Bitfield4Set(1)
+		, Bitfield5Set(1)
+		, Bitfield6(0)
+		, Bitfield7Set(1)
 	{ }
 
 	/** Creates an uninitialized instance. */
