- Reduced the CBOR encoded size of TArray<uint8>/TArray<int8> by a factor of two (approximatively) by encoding the bytes as byte string rather than CBOR arrays of numbers.
  - The CborStructDeserializerBackend is fully backward compatible and can read the new and the old format without any code change.
  - The CborStructSerializerBackend default to the new format, but can be configured to output the old format.

This partially addresses #jira UE-78722 - Potential Memory Leak with Disaster Recovery Plugin
  - Sending large packages (containing 4K texture) to Disaster Recovery use about half the memory and bandwith previously required and greatly reduce the memory footprint of Disaster Recovery.

diff --git a/Engine/Plugins/Messaging/UdpMessaging/Source/UdpMessaging/Private/Transport/UdpDeserializedMessage.cpp b/Engine/Plugins/Messaging/UdpMessaging/Source/UdpMessaging/Private/Transport/UdpDeserializedMessage.cpp
--- a/Engine/Plugins/Messaging/UdpMessaging/Source/UdpMessaging/Private/Transport/UdpDeserializedMessage.cpp
+++ b/Engine/Plugins/Messaging/UdpMessaging/Source/UdpMessaging/Private/Transport/UdpDeserializedMessage.cpp
@@ -42,7 +42,7 @@ class FUdpDeserializedMessageDetails
 {
 public:
 	static bool DeserializeV10(FUdpDeserializedMessage& DeserializedMessage, FMemoryReader& MessageReader);
-	static bool DeserializeV11_12(FUdpDeserializedMessage& DeserializedMessage, FMemoryReader& MessageReader);
+	static bool DeserializeV11_13(FUdpDeserializedMessage& DeserializedMessage, FMemoryReader& MessageReader);
 	static bool Deserialize(FUdpDeserializedMessage& DeserializedMessage, const FUdpReassembledMessage& ReassembledMessage);
 };
 
@@ -231,7 +231,7 @@ bool FUdpDeserializedMessageDetails::DeserializeV10(FUdpDeserializedMessage& Des
 	return FStructDeserializer::Deserialize(DeserializedMessage.MessageData, *DeserializedMessage.TypeInfo, Backend);
 }
 
-bool FUdpDeserializedMessageDetails::DeserializeV11_12(FUdpDeserializedMessage& DeserializedMessage, FMemoryReader& MessageReader)
+bool FUdpDeserializedMessageDetails::DeserializeV11_13(FUdpDeserializedMessage& DeserializedMessage, FMemoryReader& MessageReader)
 {
 	// message type info
 	{
@@ -369,7 +369,9 @@ bool FUdpDeserializedMessageDetails::Deserialize(FUdpDeserializedMessage& Deseri
 	case 11:
 		// fallthrough
 	case 12:
-		return DeserializeV11_12(DeserializedMessage, MessageReader);
+		// fallthrough
+	case 13:
+		return DeserializeV11_13(DeserializedMessage, MessageReader);
 		break;
 
 	default:
diff --git a/Engine/Plugins/Messaging/UdpMessaging/Source/UdpMessaging/Private/Transport/UdpMessageProcessor.cpp b/Engine/Plugins/Messaging/UdpMessaging/Source/UdpMessaging/Private/Transport/UdpMessageProcessor.cpp
--- a/Engine/Plugins/Messaging/UdpMessaging/Source/UdpMessaging/Private/Transport/UdpMessageProcessor.cpp
+++ b/Engine/Plugins/Messaging/UdpMessaging/Source/UdpMessaging/Private/Transport/UdpMessageProcessor.cpp
@@ -175,9 +175,10 @@ bool FUdpMessageProcessor::Init()
 	Beacon = new FUdpMessageBeacon(Socket, LocalNodeId, MulticastEndpoint);
 	SocketSender = new FUdpSocketSender(Socket, TEXT("FUdpMessageProcessor.Sender"));
 
-	// Current protocol version 12
+	// Current protocol version 13
 	SupportedProtocolVersions.Add(UDP_MESSAGING_TRANSPORT_PROTOCOL_VERSION);
-	// Support Protocol version 10-11
+	// Support Protocol version 10, 11, 12
+	SupportedProtocolVersions.Add(12);
 	SupportedProtocolVersions.Add(11);
 	SupportedProtocolVersions.Add(10);
 
diff --git a/Engine/Plugins/Messaging/UdpMessaging/Source/UdpMessaging/Private/Transport/UdpSerializeMessageTask.cpp b/Engine/Plugins/Messaging/UdpMessaging/Source/UdpMessaging/Private/Transport/UdpSerializeMessageTask.cpp
--- a/Engine/Plugins/Messaging/UdpMessaging/Source/UdpMessaging/Private/Transport/UdpSerializeMessageTask.cpp
+++ b/Engine/Plugins/Messaging/UdpMessaging/Source/UdpMessaging/Private/Transport/UdpSerializeMessageTask.cpp
@@ -49,8 +49,8 @@ void SerializeMessageV10(FArchive& Archive, const TSharedRef<IMessageContext, ES
 	FStructSerializer::Serialize(MessageContext->GetMessage(), *MessageContext->GetMessageTypeInfo(), Backend);
 }
 
-/** Serialization Routine for message using Protocol version 11 or 12 */
-void SerializeMessageV11_12(FArchive& Archive, const TSharedRef<IMessageContext, ESPMode::ThreadSafe>& MessageContext, const EStructSerializerBackendFlags StructSerializerBackendFlags)
+/** Serialization Routine for message using Protocol version 11, 12 or 13. */
+void SerializeMessageV11_13(FArchive& Archive, const TSharedRef<IMessageContext, ESPMode::ThreadSafe>& MessageContext, const EStructSerializerBackendFlags StructSerializerBackendFlags)
 {
 	const FName& MessageType = MessageContext->GetMessageType();
 	Archive << const_cast<FName&>(MessageType);
@@ -118,12 +118,17 @@ void FUdpSerializeMessageTask::DoTask(ENamedThreads::Type CurrentThread, const F
 
 			case 11:
 				ProtocolMaxSegmentSize = UDP_MESSAGING_SEGMENT_SIZE * (int64)UINT16_MAX;
-				UdpSerializeMessageTaskDetails::SerializeMessageV11_12(Archive, MessageContext, EStructSerializerBackendFlags::Legacy);
+				UdpSerializeMessageTaskDetails::SerializeMessageV11_13(Archive, MessageContext, EStructSerializerBackendFlags::Legacy);
 				break;
 
 			case 12:
 				ProtocolMaxSegmentSize = UDP_MESSAGING_SEGMENT_SIZE * (int64)INT32_MAX;
-				UdpSerializeMessageTaskDetails::SerializeMessageV11_12(Archive, MessageContext, EStructSerializerBackendFlags::Default);
+				UdpSerializeMessageTaskDetails::SerializeMessageV11_13(Archive, MessageContext, EStructSerializerBackendFlags::WriteTextAsComplexString);
+				break;
+
+			case 13:
+				ProtocolMaxSegmentSize = UDP_MESSAGING_SEGMENT_SIZE * (int64)INT32_MAX;
+				UdpSerializeMessageTaskDetails::SerializeMessageV11_13(Archive, MessageContext, EStructSerializerBackendFlags::Default);
 				break;
 
 			default:
diff --git a/Engine/Plugins/Messaging/UdpMessaging/Source/UdpMessaging/Private/UdpMessagingPrivate.h b/Engine/Plugins/Messaging/UdpMessaging/Source/UdpMessaging/Private/UdpMessagingPrivate.h
--- a/Engine/Plugins/Messaging/UdpMessaging/Source/UdpMessaging/Private/UdpMessagingPrivate.h
+++ b/Engine/Plugins/Messaging/UdpMessaging/Source/UdpMessaging/Private/UdpMessagingPrivate.h
@@ -26,8 +26,11 @@ DECLARE_LOG_CATEGORY_EXTERN(LogUdpMessaging, Log, All);
 /** Define the desired size for the message segments */
 #define UDP_MESSAGING_SEGMENT_SIZE 1024
 
-/** Defines the protocol version of the UDP message transport. */
-#define UDP_MESSAGING_TRANSPORT_PROTOCOL_VERSION 12
+/**
+ * Defines the protocol version of the UDP message transport.
+ * @note When changing the version, ensure to update the serialization/deserialization code in UdpSerializeMessageTask.cpp/UdpDeserializedMessage.cpp and the supported versions in FUdpMessageProcessor::Init().
+ */
+#define UDP_MESSAGING_TRANSPORT_PROTOCOL_VERSION 13 
 
 /** Supported message serialization format */
 enum class EUdpMessageFormat : uint8
diff --git a/Engine/Source/Runtime/Cbor/Private/CborWriter.cpp b/Engine/Source/Runtime/Cbor/Private/CborWriter.cpp
--- a/Engine/Source/Runtime/Cbor/Private/CborWriter.cpp
+++ b/Engine/Source/Runtime/Cbor/Private/CborWriter.cpp
@@ -122,6 +122,12 @@ void FCborWriter::WriteValue(const char* CString, uint64 Length)
 	Stream->Serialize(const_cast<char*>(CString), Length);
 }
 
+void FCborWriter::WriteValue(const uint8* Bytes, uint64 Length)
+{
+	static_assert(sizeof(uint8) == sizeof(char), "Expected char type to be 1 byte");
+	WriteValue(reinterpret_cast<const char*>(Bytes), Length);
+}
+
 FCborHeader FCborWriter::WriteUIntValue(FCborHeader Header, FArchive& Ar, uint64 Value)
 {
 	if (Value < 24)
diff --git a/Engine/Source/Runtime/Cbor/Private/Tests/CborTests.cpp b/Engine/Source/Runtime/Cbor/Private/Tests/CborTests.cpp
--- a/Engine/Source/Runtime/Cbor/Private/Tests/CborTests.cpp
+++ b/Engine/Source/Runtime/Cbor/Private/Tests/CborTests.cpp
@@ -199,6 +199,14 @@ bool FCborAutomationTest::RunTest(const FString& Parameters)
 	check(Context.MajorType() == ECborCode::ByteString);
 	check(TCString<char>::Strcmp(Context.AsCString(), TestCString) == 0);
 
+	// Byte String. (with '\0' in the middle)
+	uint8 ByteString[] = {-1, -55, -128, 0, 1, 15, 127};
+	Writer.WriteValue(ByteString, sizeof(ByteString)/sizeof(uint8));
+	check(Reader.ReadNext(Context) == true);
+	check(Context.MajorType() == ECborCode::ByteString);
+	check(FMemory::Memcmp(ByteString, Context.AsByteArray().GetData(), sizeof(ByteString)/sizeof(uint8)) == 0);
+	check(Context.AsByteArray().Num() == sizeof(ByteString)/sizeof(uint8));
+
 	// Array
 	TArray<int64> IntArray { 0, 1, -1, 10, -1000, -3000000000LL, 240, -24 };
 	Writer.WriteContainerStart(ECborCode::Array, IntArray.Num());
diff --git a/Engine/Source/Runtime/Cbor/Public/CborTypes.h b/Engine/Source/Runtime/Cbor/Public/CborTypes.h
--- a/Engine/Source/Runtime/Cbor/Public/CborTypes.h
+++ b/Engine/Source/Runtime/Cbor/Public/CborTypes.h
@@ -262,6 +262,13 @@ struct FCborContext
 		return RawTextValue.GetData();
 	}
 
+	/** @return the context as a raw byte array. */
+	TArrayView<const uint8> AsByteArray() const
+	{
+		check(MajorType() == ECborCode::ByteString);
+		return MakeArrayView(reinterpret_cast<const uint8*>(RawTextValue.GetData()), RawTextValue.Num() - 1); // -1 to exclude the null terminated added by the reader.
+	}
+
 private:
 	friend class FCborReader;
 	friend class FCborWriter;
diff --git a/Engine/Source/Runtime/Cbor/Public/CborWriter.h b/Engine/Source/Runtime/Cbor/Public/CborWriter.h
--- a/Engine/Source/Runtime/Cbor/Public/CborWriter.h
+++ b/Engine/Source/Runtime/Cbor/Public/CborWriter.h
@@ -39,6 +39,7 @@ public:
 	void WriteValue(double Value);
 	void WriteValue(const FString& Value);
 	void WriteValue(const char* CString, uint64 Length);
+	void WriteValue(const uint8* Bytes, uint64 Length);
 
 private:
 	/** Write a uint Value for Header in Ar and return the final generated cbor Header. */
diff --git a/Engine/Source/Runtime/Serialization/Private/Backends/CborStructDeserializerBackend.cpp b/Engine/Source/Runtime/Serialization/Private/Backends/CborStructDeserializerBackend.cpp
--- a/Engine/Source/Runtime/Serialization/Private/Backends/CborStructDeserializerBackend.cpp
+++ b/Engine/Source/Runtime/Serialization/Private/Backends/CborStructDeserializerBackend.cpp
@@ -35,6 +35,21 @@ bool FCborStructDeserializerBackend::GetNextToken(EStructDeserializerBackendToke
 {
 	LastMapKey.Reset();
 
+	if (bDeserializingByteArray) // Deserializing the content of a TArray<uint8>/TArray<int8> property?
+	{
+		if (DeserializingByteArrayIndex < LastContext.AsByteArray().Num())
+		{
+			OutToken = EStructDeserializerBackendTokens::Property; // Need to consume a byte from the CBOR ByteString as a UByteProperty/UInt8Property.
+		}
+		else
+		{
+			bDeserializingByteArray = false;
+			OutToken = EStructDeserializerBackendTokens::ArrayEnd; // All bytes from the byte string were deserialized into the TArray<uint8>/TArray<int8>.
+		}
+
+		return true;
+	}
+
 	if (!CborReader.ReadNext(LastContext))
 	{
 		OutToken = LastContext.IsError() ? EStructDeserializerBackendTokens::Error : EStructDeserializerBackendTokens::None;
@@ -73,6 +88,11 @@ bool FCborStructDeserializerBackend::GetNextToken(EStructDeserializerBackendToke
 	case ECborCode::Map:
 		OutToken = EStructDeserializerBackendTokens::StructureStart;
 		break;
+	case ECborCode::ByteString: // Used for size optimization on TArray<uint8>/TArray<int8>. Might be replaced if https://datatracker.ietf.org/doc/draft-ietf-cbor-array-tags/ is adopted.
+		OutToken = EStructDeserializerBackendTokens::ArrayStart;
+		DeserializingByteArrayIndex = 0;
+		bDeserializingByteArray = true;
+		break;
 	case ECborCode::Int:
 		// fall through
 	case ECborCode::Uint:
@@ -141,7 +161,6 @@ bool FCborStructDeserializerBackend::ReadProperty(UProperty* Property, UProperty
 			return StructDeserializerBackendUtilities::SetPropertyValue(Int64Property, Outer, Data, ArrayIndex, (int64)LastContext.AsInt());
 		}
 
-
 		UE_LOG(LogSerialization, Verbose, TEXT("Integer field %s with value '%d' is not supported in UProperty type %s (%s)"), *Property->GetFName().ToString(), LastContext.AsUInt(), *Property->GetClass()->GetName(), *GetDebugString());
 
 		return false;
@@ -237,6 +256,30 @@ bool FCborStructDeserializerBackend::ReadProperty(UProperty* Property, UProperty
 	}
 	break;
 
+	// Stream of bytes: Used for TArray<uint8>/TArray<int8>
+	case ECborCode::ByteString:
+	{
+		check(bDeserializingByteArray);
+
+		// Consume one byte from the byte string.
+		TArrayView<const uint8> DeserializedByteArray = LastContext.AsByteArray();
+		check(DeserializingByteArrayIndex < DeserializedByteArray.Num());
+		uint8 ByteValue = DeserializedByteArray[DeserializingByteArrayIndex++];
+
+		if (UByteProperty* ByteProperty = Cast<UByteProperty>(Property))
+		{
+			return StructDeserializerBackendUtilities::SetPropertyValue(ByteProperty, Outer, Data, ArrayIndex, ByteValue);
+		}
+		else if (UInt8Property* Int8Property = Cast<UInt8Property>(Property))
+		{
+			return StructDeserializerBackendUtilities::SetPropertyValue(Int8Property, Outer, Data, ArrayIndex, (int8)ByteValue);
+		}
+
+		UE_LOG(LogSerialization, Verbose, TEXT("Error while deserializing field %s. Unexpected UProperty type %s. Expected a UByteProperty/UInt8Property to deserialize a TArray<uint8>/TArray<int8>"), *Property->GetFName().ToString(), *Property->GetClass()->GetName());
+		return false;
+	}
+	break;
+
 	// Prim
 	case ECborCode::Prim:
 	{
@@ -283,12 +326,19 @@ bool FCborStructDeserializerBackend::ReadProperty(UProperty* Property, UProperty
 	}
 
 	return true;
-
 }
 
 void FCborStructDeserializerBackend::SkipArray()
 {
-	CborReader.SkipContainer(ECborCode::Array);
+	if (bDeserializingByteArray) // Deserializing a TArray<uint8>/TArray<int8> property as byte string?
+	{
+		check(DeserializingByteArrayIndex == 0);
+		bDeserializingByteArray = false;
+	}
+	else
+	{
+		CborReader.SkipContainer(ECborCode::Array);
+	}
 }
 
 void FCborStructDeserializerBackend::SkipStructure()
diff --git a/Engine/Source/Runtime/Serialization/Private/Backends/CborStructSerializerBackend.cpp b/Engine/Source/Runtime/Serialization/Private/Backends/CborStructSerializerBackend.cpp
--- a/Engine/Source/Runtime/Serialization/Private/Backends/CborStructSerializerBackend.cpp
+++ b/Engine/Source/Runtime/Serialization/Private/Backends/CborStructSerializerBackend.cpp
@@ -22,10 +22,25 @@ void FCborStructSerializerBackend::BeginArray(const FStructSerializerState& Stat
 {
 	UObject* Outer = State.ValueProperty->GetOuter();
 
+	// If TArray<uint8>/TArray<int8> content needs to be written as ByteString (to prevent paying a 1 byte header for each byte greater than 23 required by CBOR array).
+	if (EnumHasAnyFlags(Flags, EStructSerializerBackendFlags::WriteByteArrayAsByteStream))
+	{
+		if (UArrayProperty* ArrayProperty = Cast<UArrayProperty>(State.ValueProperty))
+		{
+			if (Cast<UByteProperty>(ArrayProperty->Inner) || Cast<UInt8Property>(ArrayProperty->Inner)) // A CBOR draft to support homogeneous array exists, but is not yet approved: https://datatracker.ietf.org/doc/draft-ietf-cbor-array-tags/.
+			{
+				check(!bSerializingByteArray);
+				FScriptArrayHelper ArrayHelper(ArrayProperty, ArrayProperty->ContainerPtrToValuePtr<void>(State.ValueData));
+				AccumulatedBytes.Reset(ArrayHelper.Num());
+				bSerializingByteArray = true;
+			}
+		}
+	}
+
 	// Array nested in Array/Set
 	if ((Outer != nullptr) && (Outer->GetClass() == UArrayProperty::StaticClass() || Outer->GetClass() == USetProperty::StaticClass()))
 	{
-		CborWriter.WriteContainerStart(ECborCode::Array, -1);
+		// fall through.
 	}
 	// Array nested in Map
 	else if (State.KeyProperty != nullptr)
@@ -33,12 +48,15 @@ void FCborStructSerializerBackend::BeginArray(const FStructSerializerState& Stat
 		FString KeyString;
 		State.KeyProperty->ExportTextItem(KeyString, State.KeyData, nullptr, nullptr, PPF_None);
 		CborWriter.WriteValue(KeyString);
-		CborWriter.WriteContainerStart(ECborCode::Array, -1/*Indefinite*/);
 	}
 	// Array nested in Object
 	else
 	{
 		CborWriter.WriteValue(State.ValueProperty->GetName());
+	}
+
+	if (!bSerializingByteArray) // TArray<uint8>/TArray<int8> are written as ByteString rather than CBOR array because it is more size efficient.
+	{
 		CborWriter.WriteContainerStart(ECborCode::Array, -1/*Indefinite*/);
 	}
 }
@@ -78,7 +96,16 @@ void FCborStructSerializerBackend::BeginStructure(const FStructSerializerState&
 
 void FCborStructSerializerBackend::EndArray(const FStructSerializerState& State)
 {
-	CborWriter.WriteContainerEnd();
+	if (bSerializingByteArray) // Does end a TArray<uint8>/TArray<int8>?
+	{
+		// Flush the accumulated bytes as a ByteString().
+		CborWriter.WriteValue(AccumulatedBytes.GetData(), AccumulatedBytes.Num());
+		bSerializingByteArray = false;
+	}
+	else
+	{
+		CborWriter.WriteContainerEnd();
+	}
 }
 
 void FCborStructSerializerBackend::EndStructure(const FStructSerializerState& State)
@@ -169,6 +196,10 @@ void FCborStructSerializerBackend::WriteProperty(const FStructSerializerState& S
 		{
 			WritePropertyValue(CborWriter, State, ByteProperty->Enum->GetNameStringByValue(ByteProperty->GetPropertyValue_InContainer(State.ValueData, ArrayIndex)));
 		}
+		else if (bSerializingByteArray) // Writing a byte from a TArray<uint8>/TArray<int8>?
+		{
+			AccumulatedBytes.Add(ByteProperty->GetPropertyValue_InContainer(State.ValueData, ArrayIndex));
+		}
 		else
 		{
 			WritePropertyValue(CborWriter, State, (int64)ByteProperty->GetPropertyValue_InContainer(State.ValueData, ArrayIndex));
@@ -192,7 +223,15 @@ void FCborStructSerializerBackend::WriteProperty(const FStructSerializerState& S
 	}
 	else if (State.ValueType == UInt8Property::StaticClass())
 	{
-		WritePropertyValue(CborWriter, State, (int64)CastChecked<UInt8Property>(State.ValueProperty)->GetPropertyValue_InContainer(State.ValueData, ArrayIndex));
+		int8 Value = CastChecked<UInt8Property>(State.ValueProperty)->GetPropertyValue_InContainer(State.ValueData, ArrayIndex);
+		if (bSerializingByteArray) // Writing a int8 from a TArray<uint8>/TArray<int8>?
+		{
+			AccumulatedBytes.Add(Value);
+		}
+		else
+		{
+			WritePropertyValue(CborWriter, State, (int64)Value);
+		}
 	}
 	else if (State.ValueType == UInt16Property::StaticClass())
 	{
diff --git a/Engine/Source/Runtime/Serialization/Private/Tests/StructSerializerTest.cpp b/Engine/Source/Runtime/Serialization/Private/Tests/StructSerializerTest.cpp
--- a/Engine/Source/Runtime/Serialization/Private/Tests/StructSerializerTest.cpp
+++ b/Engine/Source/Runtime/Serialization/Private/Tests/StructSerializerTest.cpp
@@ -98,9 +98,9 @@ namespace StructSerializerTest
 		Test.TestEqual<FQuat>(TEXT("Builtins.Quat must be the same before and after de-/serialization"), TestStruct.Builtins.Quat, TestStruct2.Builtins.Quat);
 		Test.TestEqual<FColor>(TEXT("Builtins.Color must be the same before and after de-/serialization"), TestStruct.Builtins.Color, TestStruct2.Builtins.Color);
 
-
 		// test arrays
 		Test.TestEqual<TArray<int32>>(TEXT("Arrays.Int32Array must be the same before and after de-/serialization"), TestStruct.Arrays.Int32Array, TestStruct2.Arrays.Int32Array);
+		Test.TestEqual<TArray<uint8>>(TEXT("Arrays.ByteArray must be the same before and after de-/serialization"), TestStruct.Arrays.ByteArray, TestStruct2.Arrays.ByteArray);
 		Test.TestEqual<int32>(TEXT("Arrays.StaticSingleElement[0] must be the same before and after de-/serialization"), TestStruct.Arrays.StaticSingleElement[0], TestStruct2.Arrays.StaticSingleElement[0]);
 		Test.TestEqual<int32>(TEXT("Arrays.StaticInt32Array[0] must be the same before and after de-/serialization"), TestStruct.Arrays.StaticInt32Array[0], TestStruct2.Arrays.StaticInt32Array[0]);
 		Test.TestEqual<int32>(TEXT("Arrays.StaticInt32Array[1] must be the same before and after de-/serialization"), TestStruct.Arrays.StaticInt32Array[1], TestStruct2.Arrays.StaticInt32Array[1]);
@@ -163,4 +163,129 @@ bool FStructSerializerTest::RunTest( const FString& Parameters )
 	return true;
 }
 
+IMPLEMENT_SIMPLE_AUTOMATION_TEST(FStructSerializerCborByteArrayTest, "System.Core.Serialization.StructSerializerCborByteArray", EAutomationTestFlags::EditorContext | EAutomationTestFlags::EngineFilter)
+
+bool FStructSerializerCborByteArrayTest::RunTest( const FString& Parameters )
+{
+	// Ensure TArray<uint8>/TArray<int8> are written as CBOR byte string (~2x more compact) by default rather than a CBOR array.
+	{
+		static_assert((EStructSerializerBackendFlags::Default & EStructSerializerBackendFlags::WriteByteArrayAsByteStream) == EStructSerializerBackendFlags::WriteByteArrayAsByteStream, "Test below expects 'EStructSerializerBackendFlags::Default' to contain 'EStructSerializerBackendFlags::WriteByteArrayAsByteStream'");
+		
+		// Serialization
+		TArray<uint8> Buffer;
+		FMemoryWriter Writer(Buffer);
+		FCborStructSerializerBackend SerializerBackend(Writer, EStructSerializerBackendFlags::Default);
+		FStructSerializerByteArray WrittenStruct;
+		FStructSerializer::Serialize(WrittenStruct, SerializerBackend);
+
+		TestTrue(TEXT("Arrays of int8/uint8 must be encoded in byte string (compact)"), Buffer.Num() == 54); // Copy the 54 bytes from VC++ Memory viewer to CBOR playground http://cbor.me/ to validate the count/content.
+
+		// Deserialization
+		FMemoryReader Reader(Buffer);
+		FCborStructDeserializerBackend DeserializerBackend(Reader);
+		FStructDeserializerPolicies Policies;
+		Policies.MissingFields = EStructDeserializerErrorPolicies::Warning;
+		FStructSerializerByteArray ReadStruct(NoInit);
+		FStructDeserializer::Deserialize(ReadStruct, DeserializerBackend, Policies);
+
+		TestTrue(TEXT("Value before TArray<uint8> must be the same before and after de-/serialization."), ReadStruct.Dummy1 == 1);
+		TestTrue(TEXT("Value after TArray<uint8> must be the same before and after de-/serialization."), ReadStruct.Dummy2 == 2);
+		TestTrue(TEXT("Value after TArray<int8> must be the same before and after de-/serialization."), ReadStruct.Dummy3 == 3);
+		TestTrue(TEXT("Array uint8 must be the same before and after de-/serialization"), WrittenStruct.ByteArray == ReadStruct.ByteArray);
+		TestTrue(TEXT("Array int8 must be the same before and after de-/serialization"), WrittenStruct.Int8Array == ReadStruct.Int8Array);
+	}
+
+	// Ensure TArray<uint8>/TArray<int8> encoded in CBOR byte string are skipped on deserialization if required by the policy.
+	{
+		// Serialization
+		TArray<uint8> Buffer;
+		FMemoryWriter Writer(Buffer);
+		FCborStructSerializerBackend SerializerBackend(Writer, EStructSerializerBackendFlags::Default);
+		FStructSerializerByteArray WrittenStruct;
+		FStructSerializer::Serialize(WrittenStruct, SerializerBackend);
+
+		// Deserialization
+		FMemoryReader Reader(Buffer);
+		FCborStructDeserializerBackend DeserializerBackend(Reader);
+
+		// Skip the array properties named "ByteArray" and "Int8Array".
+		FStructDeserializerPolicies Policies;
+		Policies.PropertyFilter = [](const UProperty* CurrentProp, const UProperty* ParentProp)
+		{
+			const bool bFilteredOut = (CurrentProp->GetFName() == FName(TEXT("ByteArray")) || CurrentProp->GetFName() == FName(TEXT("Int8Array")));
+			return !bFilteredOut;
+		};
+
+		Policies.MissingFields = EStructDeserializerErrorPolicies::Warning;
+		FStructSerializerByteArray ReadStruct(NoInit);
+		FStructDeserializer::Deserialize(ReadStruct, DeserializerBackend, Policies);
+
+		TestTrue(TEXT("Per deserializer policy, value before TArray<uint8> must be the same before and after de-/serialization."), ReadStruct.Dummy1 == 1);
+		TestTrue(TEXT("Per deserializer policy, value after TArray<uint8> must be the same before and after de-/serialization."), ReadStruct.Dummy2 == 2);
+		TestTrue(TEXT("Per deserializer policy, value after TArray<int8> must be the same before and after de-/serialization."), ReadStruct.Dummy3 == 3);
+		TestTrue(TEXT("Per deserializer policy, TArray<uint8> must be skipped on deserialization"), ReadStruct.ByteArray.Num() == 0);
+		TestTrue(TEXT("Per deserializer policy, TArray<int8> must be skipped on deserialization"), ReadStruct.Int8Array.Num() == 0);
+	}
+
+	// Ensure empty TArray<uint8>/TArray<int8> are written as zero-length CBOR byte string.
+	{
+		// Serialization
+		TArray<uint8> Buffer;
+		FMemoryWriter Writer(Buffer);
+		FCborStructSerializerBackend SerializerBackend(Writer, EStructSerializerBackendFlags::Default);
+		FStructSerializerByteArray WrittenStruct(NoInit); // Keep the TArray<> empty.
+		WrittenStruct.Dummy1 = 1;
+		WrittenStruct.Dummy2 = 2;
+		WrittenStruct.Dummy3 = 3;
+		FStructSerializer::Serialize(WrittenStruct, SerializerBackend);
+
+		TestTrue(TEXT("Arrays of int8/uint8 must be encoded in byte string (compact)"), Buffer.Num() == 48); // Copy the 48 bytes from VC++ Memory viewer to CBOR playground http://cbor.me/ to validate the count/content.
+
+		// Deserialization
+		FMemoryReader Reader(Buffer);
+		FCborStructDeserializerBackend DeserializerBackend(Reader);
+		FStructDeserializerPolicies Policies;
+		Policies.MissingFields = EStructDeserializerErrorPolicies::Warning;
+		FStructSerializerByteArray ReadStruct(NoInit);
+		FStructDeserializer::Deserialize(ReadStruct, DeserializerBackend, Policies);
+
+		TestTrue(TEXT("Value before TArray<uint8> must be the same before and after de-/serialization."), ReadStruct.Dummy1 == 1);
+		TestTrue(TEXT("Value after TArray<uint8> must be the same before and after de-/serialization."), ReadStruct.Dummy2 == 2);
+		TestTrue(TEXT("Value after TArray<int8> must be the same before and after de-/serialization."), ReadStruct.Dummy3 == 3);
+		TestTrue(TEXT("Array uint8 must be the same before and after de-/serialization"), WrittenStruct.ByteArray == ReadStruct.ByteArray);
+		TestTrue(TEXT("Array int8 must be the same before and after de-/serialization"), WrittenStruct.Int8Array == ReadStruct.Int8Array);
+	}
+
+	// Ensure TArray<uint8>/TArray<int8> CBOR serialization is backward compatible. (Serializer can write the old format and deserializer can read it)
+	{
+		static_assert((EStructSerializerBackendFlags::Legacy & EStructSerializerBackendFlags::WriteByteArrayAsByteStream) == EStructSerializerBackendFlags::None, "Test below expects 'EStructSerializerBackendFlags::Legacy' to not have 'EStructSerializerBackendFlags::WriteByteArrayAsByteStream'");
+		
+		// Serialize TArray<uint8>/TArray<int8> it they were prior 4.25. (CBOR array rather than CBOR byte string)
+		TArray<uint8> Buffer;
+		FMemoryWriter Writer(Buffer);
+		FCborStructSerializerBackend SerializerBackend(Writer, EStructSerializerBackendFlags::Legacy); // Legacy mode doesn't enable EStructSerializerBackendFlags::WriteByteArrayAsByteStream.
+		FStructSerializerByteArray WrittenStruct;
+		FStructSerializer::Serialize(WrittenStruct, SerializerBackend);
+
+		TestTrue(TEXT("Backward compatibility: Serialized size check"), Buffer.Num() == 60); // Copy the 60 bytes from VC++ Memory viewer to CBOR playground http://cbor.me/ to validate the count/content.
+
+		// Deserialize TArray<uint8>/TArray<int8> as they were prior 4.25.
+		FMemoryReader Reader(Buffer);
+		FCborStructDeserializerBackend DeserializerBackend(Reader);
+		FStructDeserializerPolicies Policies;
+		Policies.MissingFields = EStructDeserializerErrorPolicies::Warning;
+		FStructSerializerByteArray ReadStruct(NoInit);
+		FStructDeserializer::Deserialize(ReadStruct, DeserializerBackend, Policies);
+
+		TestTrue(TEXT("Backward compatibility: Integer must be the same before and after de-/serialization."), ReadStruct.Dummy1 == 1);
+		TestTrue(TEXT("Backward compatibility: Integer must be the same before and after de-/serialization."), ReadStruct.Dummy2 == 2);
+		TestTrue(TEXT("Backward compatibility: Integer must be the same before and after de-/serialization."), ReadStruct.Dummy3 == 3);
+		TestTrue(TEXT("Backward compatibility: TArray<uint8> must be readable as CBOR array of number."), WrittenStruct.ByteArray == ReadStruct.ByteArray);
+		TestTrue(TEXT("Backward compatibility: TArray<int8> must be readable as CBOR array of number."), WrittenStruct.Int8Array == ReadStruct.Int8Array);
+	}
+
+	return true;
+}
+
+
 #endif //WITH_DEV_AUTOMATION_TESTS
diff --git a/Engine/Source/Runtime/Serialization/Private/Tests/StructSerializerTestTypes.h b/Engine/Source/Runtime/Serialization/Private/Tests/StructSerializerTestTypes.h
--- a/Engine/Source/Runtime/Serialization/Private/Tests/StructSerializerTestTypes.h
+++ b/Engine/Source/Runtime/Serialization/Private/Tests/StructSerializerTestTypes.h
@@ -242,6 +242,48 @@ FORCEINLINE uint32 GetTypeHash(const FStructSerializerBuiltinTestStruct& S)
 	return GetTypeHash(S.String);
 }
 
+/**
+ * Test structure for byte array properties.
+ */
+USTRUCT()
+struct FStructSerializerByteArray
+{
+	GENERATED_BODY()
+
+	UPROPERTY()
+	int32 Dummy1; // To test 'skip array'.
+
+	UPROPERTY()
+	TArray<uint8> ByteArray;
+
+	UPROPERTY()
+	int32 Dummy2; // To test 'skip array'.
+
+	UPROPERTY()
+	TArray<int8> Int8Array;
+
+	UPROPERTY()
+	int32 Dummy3; // To test 'skip array'.
+
+	FStructSerializerByteArray()
+	{
+		Dummy1 = 1;
+		Dummy2 = 2;
+		Dummy3 = 3;
+
+		ByteArray.Add(0);
+		ByteArray.Add(127);
+		ByteArray.Add(255);
+
+		Int8Array.Add(-128);
+		Int8Array.Add(0);
+		Int8Array.Add(127);
+	}
+
+	FStructSerializerByteArray(ENoInit) { }
+};
+
+
 /**
  * Test structure for array properties.
  */
@@ -253,6 +295,9 @@ struct FStructSerializerArrayTestStruct
 	UPROPERTY()
 	TArray<int32> Int32Array;
 
+	UPROPERTY()
+	TArray<uint8> ByteArray;
+
 	UPROPERTY()
 	int32 StaticSingleElement[1];
 
@@ -275,6 +320,10 @@ struct FStructSerializerArrayTestStruct
 		Int32Array.Add(0);
 		Int32Array.Add(1);
 
+		ByteArray.Add(0);
+		ByteArray.Add(127);
+		ByteArray.Add(255);
+
 		StaticSingleElement[0] = 42;
 
 		StaticInt32Array[0] = -1;
diff --git a/Engine/Source/Runtime/Serialization/Public/Backends/CborStructDeserializerBackend.h b/Engine/Source/Runtime/Serialization/Public/Backends/CborStructDeserializerBackend.h
--- a/Engine/Source/Runtime/Serialization/Public/Backends/CborStructDeserializerBackend.h
+++ b/Engine/Source/Runtime/Serialization/Public/Backends/CborStructDeserializerBackend.h
@@ -41,4 +41,10 @@ private:
 
 	/** Holds the last map key. */
 	FString LastMapKey;
+
+	/** The index of the next byte to copy from the CBOR byte stream into the corresponding TArray<uint8>/TArray<int8> property. */
+	int32 DeserializingByteArrayIndex = 0;
+
+	/** Whether a TArray<uint8>/TArray<int8> property is being deserialized. */
+	bool bDeserializingByteArray = false;
 };
diff --git a/Engine/Source/Runtime/Serialization/Public/Backends/CborStructSerializerBackend.h b/Engine/Source/Runtime/Serialization/Public/Backends/CborStructSerializerBackend.h
--- a/Engine/Source/Runtime/Serialization/Public/Backends/CborStructSerializerBackend.h
+++ b/Engine/Source/Runtime/Serialization/Public/Backends/CborStructSerializerBackend.h
@@ -8,7 +8,6 @@
 
 /**
  * Implements a writer for UStruct serialization using Cbor.
-
  */
 class SERIALIZATION_API FCborStructSerializerBackend
 	: public IStructSerializerBackend
@@ -32,8 +31,8 @@ public:
 	 */
 	FCborStructSerializerBackend(FArchive& InArchive, const EStructSerializerBackendFlags InFlags);
 
-
 	virtual ~FCborStructSerializerBackend();
+
 public:
 
 	// IStructSerializerBackend interface
@@ -50,4 +49,10 @@ private:
 
 	/** Flags controlling the serialization behavior. */
 	EStructSerializerBackendFlags Flags;
+
+	/** Stores the accumulated bytes extracted from UByteProperty/UIntProperty when writing a TArray<uint8>/TArray<int8>. */
+	TArray<uint8> AccumulatedBytes;
+
+	/** Whether the serializer is encoding array of uint8/int8 */
+	bool bSerializingByteArray = false;
 };
diff --git a/Engine/Source/Runtime/Serialization/Public/IStructSerializerBackend.h b/Engine/Source/Runtime/Serialization/Public/IStructSerializerBackend.h
--- a/Engine/Source/Runtime/Serialization/Public/IStructSerializerBackend.h
+++ b/Engine/Source/Runtime/Serialization/Public/IStructSerializerBackend.h
@@ -21,15 +21,20 @@ enum class EStructSerializerBackendFlags
 	 */
 	WriteTextAsComplexString = 1<<0,
 
+	/**
+	 * Write TArray<uint8>/TArray<int> as byte string if possible (CBOR), starting at 4.25.
+	 */
+	WriteByteArrayAsByteStream = 1<<1,
+
 	/**
 	 * Legacy settings for backwards compatibility with code compiled prior to 4.22.
 	 */
 	Legacy = None,
 
 	/**
-	 * Default settings for code compiled for 4.22 onwards.
+	 * Default settings for code compiled for 4.25 onwards.
 	 */
-	Default = WriteTextAsComplexString,
+	Default = WriteTextAsComplexString | WriteByteArrayAsByteStream,
 };
 ENUM_CLASS_FLAGS(EStructSerializerBackendFlags);
 
