- Fixed an incorrect convertion introduced in CBOR struct serializer when UProperty where changed to FProperty.

diff --git a/Engine/Source/Runtime/Serialization/Private/Backends/CborStructSerializerBackend.cpp b/Engine/Source/Runtime/Serialization/Private/Backends/CborStructSerializerBackend.cpp
--- a/Engine/Source/Runtime/Serialization/Private/Backends/CborStructSerializerBackend.cpp
+++ b/Engine/Source/Runtime/Serialization/Private/Backends/CborStructSerializerBackend.cpp
@@ -23,7 +23,7 @@ void FCborStructSerializerBackend::BeginArray(const FStructSerializerState& Stat
 	// If TArray<uint8>/TArray<int8> content needs to be written as ByteString (to prevent paying a 1 byte header for each byte greater than 23 required by CBOR array).
 	if (EnumHasAnyFlags(Flags, EStructSerializerBackendFlags::WriteByteArrayAsByteStream))
 	{
-		if (FArrayProperty* ArrayProperty = State.ValueProperty->GetOwner<FArrayProperty>())
+		if (FArrayProperty* ArrayProperty = CastField<FArrayProperty>(State.ValueProperty))
 		{
 			if (CastField<FByteProperty>(ArrayProperty->Inner) || CastField<FInt8Property>(ArrayProperty->Inner)) // A CBOR draft to support homogeneous array exists, but is not yet approved: https://datatracker.ietf.org/doc/draft-ietf-cbor-array-tags/.
 			{
