[NO NON-PII COMMIT MESSAGE]

diff --git a/Engine/Source/Runtime/Serialization/Private/StructSerializer.cpp b/Engine/Source/Runtime/Serialization/Private/StructSerializer.cpp
--- a/Engine/Source/Runtime/Serialization/Private/StructSerializer.cpp
+++ b/Engine/Source/Runtime/Serialization/Private/StructSerializer.cpp
@@ -369,11 +369,14 @@ void FStructSerializer::SerializeElement(const void* Address, FProperty* Propert
 				{
 					const void* ValueData = CurrentState.ValueData;
 
-					FFieldVariant Outer = CurrentState.ValueProperty->GetOwnerVariant();
-					if ((Outer.ToField() == nullptr) || (Outer.ToField()->GetClass() != FArrayProperty::StaticClass()))
+					if (CurrentState.ValueProperty)
 					{
-						const int32 ContainerAddressIndex = CurrentState.ElementIndex != INDEX_NONE ? CurrentState.ElementIndex : 0;
-						ValueData = CurrentState.ValueProperty->ContainerPtrToValuePtr<void>(CurrentState.ValueData, ContainerAddressIndex);
+						FFieldVariant Outer = CurrentState.ValueProperty->GetOwnerVariant();
+						if ((Outer.ToField() == nullptr) || (Outer.ToField()->GetClass() != FArrayProperty::StaticClass()))
+						{
+							const int32 ContainerAddressIndex = CurrentState.ElementIndex != INDEX_NONE ? CurrentState.ElementIndex : 0;
+							ValueData = CurrentState.ValueProperty->ContainerPtrToValuePtr<void>(CurrentState.ValueData, ContainerAddressIndex);
+						}
 					}
 
 					Backend.BeginStructure(CurrentState);
