Fixing PVS 7.7 Issues:

warning V522: Dereferencing of the null pointer 'CurrentState.Property' might take place.

diff --git a/Engine/Source/Runtime/Serialization/Private/StructDeserializer.cpp b/Engine/Source/Runtime/Serialization/Private/StructDeserializer.cpp
--- a/Engine/Source/Runtime/Serialization/Private/StructDeserializer.cpp
+++ b/Engine/Source/Runtime/Serialization/Private/StructDeserializer.cpp
@@ -449,6 +449,7 @@ bool FStructDeserializer::DeserializeElement(void* OutAddress, UStruct& OwnerInf
 		case EStructDeserializerBackendTokens::ArrayEnd:
 		{
 			// rehash the set/maps -> we're closing them
+			check(CurrentState.Property);
 			if (FSetProperty* SetProperty = CastField<FSetProperty>(CurrentState.Property))
 			{
 				FScriptSetHelper SetHelper(SetProperty, CurrentState.Data);
@@ -459,7 +460,7 @@ bool FStructDeserializer::DeserializeElement(void* OutAddress, UStruct& OwnerInf
 				FScriptMapHelper MapHelper(MapProperty, CurrentState.Data);
 				MapHelper.Rehash();
 			}
-			else if (CurrentState.Property->ArrayDim > 1 && CurrentState.ArrayIndex < CurrentState.Property->ArrayDim)
+			else if (CurrentState.Property->ArrayDim > 1 && CurrentState.ArrayIndex < CurrentState.Property->ArrayDim) //-V522 - Property should never be null
 			{
 				// error: array entry not found in static array
 				if (Policies.MissingFields != EStructDeserializerErrorPolicies::Ignore)
