Finding and addressing codepaths that check property type by explicitly checking equality with FObjectProperty::StaticClass() or FClassProperty::StaticClass().  Changing them to instead use CastField<FObjectProperty> (which will generically handle both FObjectProperty, FClassProperty, and any of their derivatives).  Could be further genericized in the future to use FObjectPropertyBase, but don't want to add more risk by introducing that at this point.

Confirmed that now System.Core.Serialization.StructElementSerializer and System.Core.Serialization.StructSerializer unit tests are passing.

diff --git a/Engine/Plugins/Runtime/Metasound/Source/MetasoundFrontend/Private/MetasoundJsonBackend.h b/Engine/Plugins/Runtime/Metasound/Source/MetasoundFrontend/Private/MetasoundJsonBackend.h
--- a/Engine/Plugins/Runtime/Metasound/Source/MetasoundFrontend/Private/MetasoundJsonBackend.h
+++ b/Engine/Plugins/Runtime/Metasound/Source/MetasoundFrontend/Private/MetasoundJsonBackend.h
@@ -359,21 +359,11 @@ namespace Metasound
 			}
 
 			// classes & objects
-			else if (State.FieldType == FClassProperty::StaticClass())
-			{
-				UObject* const& Value = CastFieldChecked<FClassProperty>(State.ValueProperty)->GetPropertyValue_InContainer(State.ValueData, ArrayIndex);
-				WritePropertyValue(JsonWriter, State, Value ? Value->GetPathName() : FString());
-			}
 			else if (State.FieldType == FSoftClassProperty::StaticClass())
 			{
 				FSoftObjectPtr const& Value = CastFieldChecked<FSoftClassProperty>(State.ValueProperty)->GetPropertyValue_InContainer(State.ValueData, ArrayIndex);
 				WritePropertyValue(JsonWriter, State, Value.IsValid() ? Value->GetPathName() : FString());
 			}
-			else if (State.FieldType == FObjectProperty::StaticClass())
-			{
-				UObject* const& Value = CastFieldChecked<FObjectProperty>(State.ValueProperty)->GetPropertyValue_InContainer(State.ValueData, ArrayIndex);
-				WritePropertyValue(JsonWriter, State, Value ? Value->GetPathName() : FString());
-			}
 			else if (State.FieldType == FWeakObjectProperty::StaticClass())
 			{
 				FWeakObjectPtr const& Value = CastFieldChecked<FWeakObjectProperty>(State.ValueProperty)->GetPropertyValue_InContainer(State.ValueData, ArrayIndex);
@@ -384,6 +374,14 @@ namespace Metasound
 				FSoftObjectPtr const& Value = CastFieldChecked<FSoftObjectProperty>(State.ValueProperty)->GetPropertyValue_InContainer(State.ValueData, ArrayIndex);
 				WritePropertyValue(JsonWriter, State, Value.ToString());
 			}
+			else if (FObjectProperty* ObjectProperty = CastField<FObjectProperty>(State.ValueProperty))
+			{
+				// @TODO: Could this be expanded to include everything derived from FObjectPropertyBase?
+				// Generic handling for a property type derived from FObjectProperty that is obtainable as a pointer and will be stored using its path.
+				// This must come after all the more specialized handlers for object property types.
+				UObject* const Value = ObjectProperty->GetObjectPropertyValue_InContainer(State.ValueData, ArrayIndex);
+				WritePropertyValue(JsonWriter, State, Value ? Value->GetPathName() : FString());
+			}
 
 			// unsupported property type
 			else
diff --git a/Engine/Source/Editor/Blutility/Private/BlutilityMenuExtensions.cpp b/Engine/Source/Editor/Blutility/Private/BlutilityMenuExtensions.cpp
--- a/Engine/Source/Editor/Blutility/Private/BlutilityMenuExtensions.cpp
+++ b/Engine/Source/Editor/Blutility/Private/BlutilityMenuExtensions.cpp
@@ -198,10 +198,8 @@ void FBlutilityMenuExtensions::CreateActorBlutilityActionsMenu(FMenuBuilder& Men
 		LOCTEXT("ScriptedActorActionsTooltip", "Scripted actions available for the selected actors"),
 		[](const FProperty* Property) -> bool
 		{
-			const FFieldClass* ClassOfProperty = Property->GetClass();
-			if (ClassOfProperty == FObjectProperty::StaticClass())
+			if (const FObjectProperty* ObjectProperty = CastField<const FObjectProperty>(Property))
 			{
-				const FObjectProperty* ObjectProperty = CastField<const FObjectProperty>(Property);
 				return ObjectProperty->PropertyClass == AActor::StaticClass();
 			}
 
diff --git a/Engine/Source/Editor/UndoHistory/Private/Widgets/UndoHistoryUtils.h b/Engine/Source/Editor/UndoHistory/Private/Widgets/UndoHistoryUtils.h
--- a/Engine/Source/Editor/UndoHistory/Private/Widgets/UndoHistoryUtils.h
+++ b/Engine/Source/Editor/UndoHistory/Private/Widgets/UndoHistoryUtils.h
@@ -39,7 +39,7 @@ public:
 				continue;
 			}
 
-			if (Property->GetClass() == FObjectProperty::StaticClass() || Property->GetClass() == FStructProperty::StaticClass() || Property->GetClass() == FEnumProperty::StaticClass())
+			if (CastField<const FObjectProperty>(*Property) || Property->GetClass() == FStructProperty::StaticClass() || Property->GetClass() == FEnumProperty::StaticClass())
 			{
 				Property->GetCPPMacroType(ClassName);
 			}
diff --git a/Engine/Source/Runtime/Serialization/Private/Backends/CborStructSerializerBackend.cpp b/Engine/Source/Runtime/Serialization/Private/Backends/CborStructSerializerBackend.cpp
--- a/Engine/Source/Runtime/Serialization/Private/Backends/CborStructSerializerBackend.cpp
+++ b/Engine/Source/Runtime/Serialization/Private/Backends/CborStructSerializerBackend.cpp
@@ -284,21 +284,11 @@ void FCborStructSerializerBackend::WriteProperty(const FStructSerializerState& S
 	}
 
 	// Classes & Objects
-	else if (State.FieldType == FClassProperty::StaticClass())
-	{
-		UObject* const& Value = CastFieldChecked<FClassProperty>(State.ValueProperty)->GetPropertyValue_InContainer(State.ValueData, ArrayIndex);
-		WritePropertyValue(CborWriter, State, Value ? Value->GetPathName() : FString());
-	}
 	else if (State.FieldType == FSoftClassProperty::StaticClass())
 	{
 		FSoftObjectPtr const& Value = CastFieldChecked<FSoftClassProperty>(State.ValueProperty)->GetPropertyValue_InContainer(State.ValueData, ArrayIndex);
 		WritePropertyValue(CborWriter, State, Value.IsValid() ? Value->GetPathName() : FString());
 	}
-	else if (State.FieldType == FObjectProperty::StaticClass())
-	{
-		UObject* const& Value = CastFieldChecked<FObjectProperty>(State.ValueProperty)->GetPropertyValue_InContainer(State.ValueData, ArrayIndex);
-		WritePropertyValue(CborWriter, State, Value ? Value->GetPathName() : FString());
-	}
 	else if (State.FieldType == FWeakObjectProperty::StaticClass())
 	{
 		FWeakObjectPtr const& Value = CastFieldChecked<FWeakObjectProperty>(State.ValueProperty)->GetPropertyValue_InContainer(State.ValueData, ArrayIndex);
@@ -309,6 +299,14 @@ void FCborStructSerializerBackend::WriteProperty(const FStructSerializerState& S
 		FSoftObjectPtr const& Value = CastFieldChecked<FSoftObjectProperty>(State.ValueProperty)->GetPropertyValue_InContainer(State.ValueData, ArrayIndex);
 		WritePropertyValue(CborWriter, State, Value.ToString());
 	}
+	else if (FObjectProperty* ObjectProperty = CastField<FObjectProperty>(State.ValueProperty))
+	{
+		// @TODO: Could this be expanded to include everything derived from FObjectPropertyBase?
+		// Generic handling for a property type derived from FObjectProperty that is obtainable as a pointer and will be stored using its path.
+		// This must come after all the more specialized handlers for object property types.
+		UObject* const Value = ObjectProperty->GetObjectPropertyValue_InContainer(State.ValueData, ArrayIndex);
+		WritePropertyValue(CborWriter, State, Value ? Value->GetPathName() : FString());
+	}
 
 	// Unsupported
 	else
diff --git a/Engine/Source/Runtime/Serialization/Private/Backends/JsonStructSerializerBackend.cpp b/Engine/Source/Runtime/Serialization/Private/Backends/JsonStructSerializerBackend.cpp
--- a/Engine/Source/Runtime/Serialization/Private/Backends/JsonStructSerializerBackend.cpp
+++ b/Engine/Source/Runtime/Serialization/Private/Backends/JsonStructSerializerBackend.cpp
@@ -232,21 +232,11 @@ void FJsonStructSerializerBackend::WriteProperty(const FStructSerializerState& S
 	}
 
 	// classes & objects
-	else if (State.FieldType == FClassProperty::StaticClass())
-	{
-		UObject* const& Value = CastFieldChecked<FClassProperty>(State.ValueProperty)->GetPropertyValue_InContainer(State.ValueData, ArrayIndex);
-		WritePropertyValue(JsonWriter, State, Value ? Value->GetPathName() : FString());
-	}
 	else if (State.FieldType == FSoftClassProperty::StaticClass())
 	{
 		FSoftObjectPtr const& Value = CastFieldChecked<FSoftClassProperty>(State.ValueProperty)->GetPropertyValue_InContainer(State.ValueData, ArrayIndex);
 		WritePropertyValue(JsonWriter, State, Value.IsValid() ? Value->GetPathName() : FString());
 	}
-	else if (State.FieldType == FObjectProperty::StaticClass())
-	{
-		UObject* const& Value = CastFieldChecked<FObjectProperty>(State.ValueProperty)->GetPropertyValue_InContainer(State.ValueData, ArrayIndex);
-		WritePropertyValue(JsonWriter, State, Value ? Value->GetPathName() : FString());
-	}
 	else if (State.FieldType == FWeakObjectProperty::StaticClass())
 	{
 		FWeakObjectPtr const& Value = CastFieldChecked<FWeakObjectProperty>(State.ValueProperty)->GetPropertyValue_InContainer(State.ValueData, ArrayIndex);
@@ -257,6 +247,14 @@ void FJsonStructSerializerBackend::WriteProperty(const FStructSerializerState& S
 		FSoftObjectPtr const& Value = CastFieldChecked<FSoftObjectProperty>(State.ValueProperty)->GetPropertyValue_InContainer(State.ValueData, ArrayIndex);
 		WritePropertyValue(JsonWriter, State, Value.ToString());
 	}
+	else if (FObjectProperty* ObjectProperty = CastField<FObjectProperty>(State.ValueProperty))
+	{
+		// @TODO: Could this be expanded to include everything derived from FObjectPropertyBase?
+		// Generic handling for a property type derived from FObjectProperty that is obtainable as a pointer and will be stored using its path.
+		// This must come after all the more specialized handlers for object property types.
+		UObject* const Value = ObjectProperty->GetObjectPropertyValue_InContainer(State.ValueData, ArrayIndex);
+		WritePropertyValue(JsonWriter, State, Value ? Value->GetPathName() : FString());
+	}
 
 	// unsupported property type
 	else
diff --git a/Engine/Source/Runtime/Serialization/Private/Tests/StructSerializerTest.cpp b/Engine/Source/Runtime/Serialization/Private/Tests/StructSerializerTest.cpp
--- a/Engine/Source/Runtime/Serialization/Private/Tests/StructSerializerTest.cpp
+++ b/Engine/Source/Runtime/Serialization/Private/Tests/StructSerializerTest.cpp
@@ -86,10 +86,12 @@ namespace StructSerializerTest
 
 	void ValidateObjects(FAutomationTestBase& Test, const FStructSerializerObjectTestStruct& Struct1, const FStructSerializerObjectTestStruct& Struct2)
 	{
-		Test.TestEqual<class UClass*>(TEXT("Objects.Class must be the same before and after de-/serialization"), Struct1.Class, Struct2.Class);
+		Test.TestEqual<class UClass*>(TEXT("Objects.RawClass must be the same before and after de-/serialization"), Struct1.RawClass, Struct2.RawClass);
+		Test.TestEqual<class UClass*>(TEXT("Objects.WrappedClass must be the same before and after de-/serialization"), Struct1.WrappedClass, Struct2.WrappedClass);
 		Test.TestEqual<TSubclassOf<class UMetaData>>(TEXT("Objects.SubClass must be the same before and after de-/serialization"), Struct1.SubClass, Struct2.SubClass);
 		Test.TestEqual<TSoftClassPtr<class UMetaData>>(TEXT("Objects.SoftClass must be the same before and after de-/serialization"), Struct1.SoftClass, Struct2.SoftClass);
-		Test.TestEqual<UObject*>(TEXT("Objects.Object must be the same before and after de-/serialization"), Struct1.Object, Struct2.Object);
+		Test.TestEqual<UObject*>(TEXT("Objects.RawObject must be the same before and after de-/serialization"), Struct1.RawObject, Struct2.RawObject);
+		Test.TestEqual<UObject*>(TEXT("Objects.WrappedObject must be the same before and after de-/serialization"), Struct1.WrappedObject, Struct2.WrappedObject);
 		Test.TestEqual<TWeakObjectPtr<class UMetaData>>(TEXT("Objects.WeakObject must be the same before and after de-/serialization"), Struct1.WeakObject, Struct2.WeakObject);
 		Test.TestEqual<TSoftObjectPtr<class UMetaData>>(TEXT("Objects.SoftObject must be the same before and after de-/serialization"), Struct1.SoftObject, Struct2.SoftObject);
 		Test.TestEqual<FSoftClassPath>(TEXT("Objects.ClassPath must be the same before and after de-/serialization"), Struct1.ClassPath, Struct2.ClassPath);
@@ -147,10 +149,12 @@ namespace StructSerializerTest
 		UClass* MetaDataClass = LoadClass<UMetaData>(nullptr, TEXT("/Script/CoreUObject.MetaData"));
 		UMetaData* MetaDataObject = NewObject<UMetaData>();
 		// setup object tests
-		OriginalStruct.Objects.Class = MetaDataClass;
+		OriginalStruct.Objects.RawClass = MetaDataClass;
+		OriginalStruct.Objects.WrappedClass = MetaDataClass;
 		OriginalStruct.Objects.SubClass = MetaDataClass;
 		OriginalStruct.Objects.SoftClass = MetaDataClass;
-		OriginalStruct.Objects.Object = MetaDataObject;
+		OriginalStruct.Objects.RawObject = MetaDataObject;
+		OriginalStruct.Objects.WrappedObject = MetaDataObject;
 		OriginalStruct.Objects.WeakObject = MetaDataObject;
 		OriginalStruct.Objects.SoftObject = MetaDataObject;
 		OriginalStruct.Objects.ClassPath = MetaDataClass;
@@ -508,10 +512,12 @@ namespace StructSerializerTest
 		UClass* MetaDataClass = LoadClass<UMetaData>(nullptr, TEXT("/Script/CoreUObject.MetaData"));
 		UMetaData* MetaDataObject = NewObject<UMetaData>();
 		// setup object tests
-		TestStruct.Objects.Class = MetaDataClass;
+		TestStruct.Objects.RawClass = MetaDataClass;
+		TestStruct.Objects.WrappedClass = MetaDataClass;
 		TestStruct.Objects.SubClass = MetaDataClass;
 		TestStruct.Objects.SoftClass = MetaDataClass;
-		TestStruct.Objects.Object = MetaDataObject;
+		TestStruct.Objects.RawObject = MetaDataObject;
+		TestStruct.Objects.WrappedObject = MetaDataObject;
 		TestStruct.Objects.WeakObject = MetaDataObject;
 		TestStruct.Objects.SoftObject = MetaDataObject;
 		TestStruct.Objects.ClassPath = MetaDataClass;
diff --git a/Engine/Source/Runtime/Serialization/Private/Tests/StructSerializerTestTypes.h b/Engine/Source/Runtime/Serialization/Private/Tests/StructSerializerTestTypes.h
--- a/Engine/Source/Runtime/Serialization/Private/Tests/StructSerializerTestTypes.h
+++ b/Engine/Source/Runtime/Serialization/Private/Tests/StructSerializerTestTypes.h
@@ -131,7 +131,10 @@ struct FStructSerializerObjectTestStruct
 	GENERATED_BODY()
 
 	UPROPERTY()
-	TObjectPtr<UClass> Class;
+	UClass* RawClass;
+
+	UPROPERTY()
+	TObjectPtr<UClass> WrappedClass;
 
 	UPROPERTY()
 	TSubclassOf<class UMetaData> SubClass;
@@ -140,7 +143,10 @@ struct FStructSerializerObjectTestStruct
 	TSoftClassPtr<class UMetaData> SoftClass;
 
 	UPROPERTY()
-	TObjectPtr<class UObject> Object;
+	class UObject* RawObject;
+
+	UPROPERTY()
+	TObjectPtr<class UObject> WrappedObject;
 
 	UPROPERTY()
 	TWeakObjectPtr<class UMetaData> WeakObject;
@@ -156,10 +162,12 @@ struct FStructSerializerObjectTestStruct
 
 	/** Default constructor. */
 	FStructSerializerObjectTestStruct()
-		: Class(nullptr)
+		: RawClass(nullptr)
+		, WrappedClass(nullptr)
 		, SubClass(nullptr)
 		, SoftClass(nullptr)
-		, Object(nullptr)
+		, RawObject(nullptr)
+		, WrappedObject(nullptr)
 		, WeakObject(nullptr)
 		, SoftObject(nullptr)
 		, ClassPath((UClass*)nullptr)
diff --git a/Engine/Source/Runtime/WebBrowser/Private/CEF/CEFJSStructSerializerBackend.cpp b/Engine/Source/Runtime/WebBrowser/Private/CEF/CEFJSStructSerializerBackend.cpp
--- a/Engine/Source/Runtime/WebBrowser/Private/CEF/CEFJSStructSerializerBackend.cpp
+++ b/Engine/Source/Runtime/WebBrowser/Private/CEF/CEFJSStructSerializerBackend.cpp
@@ -278,13 +278,13 @@ void FCEFJSStructSerializerBackend::WriteProperty(const FStructSerializerState&
 	}
 
 	// classes & objects
-	else if (State.FieldType == FClassProperty::StaticClass())
+	else if (FClassProperty* ClassProperty = CastField<FClassProperty>(State.ValueProperty))
 	{
-		Add(State, CastFieldChecked<FClassProperty>(State.ValueProperty)->GetPropertyValue_InContainer(State.ValueData, ArrayIndex)->GetPathName());
+		Add(State, ClassProperty->GetPropertyValue_InContainer(State.ValueData, ArrayIndex)->GetPathName());
 	}
-	else if (State.FieldType == FObjectProperty::StaticClass())
+	else if (FObjectProperty* ObjectProperty = CastField<FObjectProperty>(State.ValueProperty))
 	{
-		Add(State, CastFieldChecked<FObjectProperty>(State.ValueProperty)->GetPropertyValue_InContainer(State.ValueData, ArrayIndex));
+		Add(State, ObjectProperty->GetPropertyValue_InContainer(State.ValueData, ArrayIndex));
 	}
 
 	// unsupported property type
diff --git a/Engine/Source/Runtime/WebBrowser/Private/MobileJS/MobileJSStructSerializerBackend.cpp b/Engine/Source/Runtime/WebBrowser/Private/MobileJS/MobileJSStructSerializerBackend.cpp
--- a/Engine/Source/Runtime/WebBrowser/Private/MobileJS/MobileJSStructSerializerBackend.cpp
+++ b/Engine/Source/Runtime/WebBrowser/Private/MobileJS/MobileJSStructSerializerBackend.cpp
@@ -12,9 +12,9 @@
 void FMobileJSStructSerializerBackend::WriteProperty(const FStructSerializerState& State, int32 ArrayIndex)
 {
 	// The parent class serialzes UObjects as NULLs
-	if (State.FieldType == FObjectProperty::StaticClass())
+	if (FObjectProperty* ObjectProperty = CastField<FObjectProperty>(State.ValueProperty))
 	{
-		WriteUObject(State, CastFieldChecked<FObjectProperty>(State.ValueProperty)->GetPropertyValue_InContainer(State.ValueData, ArrayIndex));
+		WriteUObject(State, ObjectProperty->GetPropertyValue_InContainer(State.ValueData, ArrayIndex));
 	}
 	// basic property type (json serializable)
 	else
diff --git a/Engine/Source/Runtime/WebBrowser/Private/Native/NativeJSStructSerializerBackend.cpp b/Engine/Source/Runtime/WebBrowser/Private/Native/NativeJSStructSerializerBackend.cpp
--- a/Engine/Source/Runtime/WebBrowser/Private/Native/NativeJSStructSerializerBackend.cpp
+++ b/Engine/Source/Runtime/WebBrowser/Private/Native/NativeJSStructSerializerBackend.cpp
@@ -10,7 +10,7 @@
 void FNativeJSStructSerializerBackend::WriteProperty(const FStructSerializerState& State, int32 ArrayIndex)
 {
 	// The parent class serialzes UObjects as NULLs
-	if (State.FieldType == FObjectProperty::StaticClass())
+	if (FObjectProperty* ObjectProperty = CastField<FObjectProperty>(State.ValueProperty))
 	{
 		WriteUObject(State, CastFieldChecked<FObjectProperty>(State.ValueProperty)->GetPropertyValue_InContainer(State.ValueData, ArrayIndex));
 	}
