- Updating MessageBus protocol and CBor serialization backend to support communication between LWC and non LWC endpoints

diff --git a/Engine/Plugins/Messaging/UdpMessaging/Source/UdpMessaging/Private/Transport/UdpDeserializedMessage.cpp b/Engine/Plugins/Messaging/UdpMessaging/Source/UdpMessaging/Private/Transport/UdpDeserializedMessage.cpp
--- a/Engine/Plugins/Messaging/UdpMessaging/Source/UdpMessaging/Private/Transport/UdpDeserializedMessage.cpp
+++ b/Engine/Plugins/Messaging/UdpMessaging/Source/UdpMessaging/Private/Transport/UdpDeserializedMessage.cpp
@@ -43,7 +43,7 @@ class FUdpDeserializedMessageDetails
 {
 public:
 	static bool DeserializeV10(FUdpDeserializedMessage& DeserializedMessage, FMemoryReader& MessageReader);
-	static bool DeserializeV11_14(FUdpDeserializedMessage& DeserializedMessage, FMemoryReader& MessageReader);
+	static bool DeserializeV11_15(FUdpDeserializedMessage& DeserializedMessage, FMemoryReader& MessageReader, bool bIsLWCBackwardCompatibilityMode);
 	static bool Deserialize(FUdpDeserializedMessage& DeserializedMessage, const FUdpReassembledMessage& ReassembledMessage);
 };
 
@@ -233,7 +233,7 @@ bool FUdpDeserializedMessageDetails::DeserializeV10(FUdpDeserializedMessage& Des
 	return FStructDeserializer::Deserialize(DeserializedMessage.MessageData, *DeserializedMessage.TypeInfo, Backend);
 }
 
-bool FUdpDeserializedMessageDetails::DeserializeV11_14(FUdpDeserializedMessage& DeserializedMessage, FMemoryReader& MessageReader)
+bool FUdpDeserializedMessageDetails::DeserializeV11_15(FUdpDeserializedMessage& DeserializedMessage, FMemoryReader& MessageReader, bool bIsLWCBackwardCompatibilityMode)
 {
 	// message type info
 	{
@@ -337,14 +337,14 @@ bool FUdpDeserializedMessageDetails::DeserializeV11_14(FUdpDeserializedMessage&
 	case EUdpMessageFormat::CborPlatformEndianness:
 	{
 		// deserialize cbor (using this platform endianness).
-		FCborStructDeserializerBackend Backend(MessageReader, ECborEndianness::Platform);
+		FCborStructDeserializerBackend Backend(MessageReader, ECborEndianness::Platform, bIsLWCBackwardCompatibilityMode);
 		return FStructDeserializer::Deserialize(DeserializedMessage.MessageData, *DeserializedMessage.TypeInfo, Backend);
 	}
 	break;
 	case EUdpMessageFormat::CborStandardEndianness:
 	{
 		// deserialize cbor (using the CBOR standard endianness - big endian).
-		FCborStructDeserializerBackend Backend(MessageReader, ECborEndianness::StandardCompliant);
+		FCborStructDeserializerBackend Backend(MessageReader, ECborEndianness::StandardCompliant, bIsLWCBackwardCompatibilityMode);
 		return FStructDeserializer::Deserialize(DeserializedMessage.MessageData, *DeserializedMessage.TypeInfo, Backend);
 	}
 	break;
@@ -383,8 +383,18 @@ bool FUdpDeserializedMessageDetails::Deserialize(FUdpDeserializedMessage& Deseri
 	case 13:
 		// fallthrough
 	case 14:
-		return DeserializeV11_14(DeserializedMessage, MessageReader);
+	{
+		// Protocol version 11-14 requires LWC backward compat mode
+		constexpr bool bIsLWCBackwardCompatibilityMode = true;
+		return DeserializeV11_15(DeserializedMessage, MessageReader, bIsLWCBackwardCompatibilityMode);
 		break;
+	}
+	case 15:
+	{
+		constexpr bool bIsLWCBackwardCompatibilityMode = false;
+		return DeserializeV11_15(DeserializedMessage, MessageReader, bIsLWCBackwardCompatibilityMode);
+		break;
+	}
 
 	default:
 		UE_LOG(LogUdpMessaging, Error, TEXT("Unsupported Protocol Version message tasked for deserialization, discarding..."));
diff --git a/Engine/Plugins/Messaging/UdpMessaging/Source/UdpMessaging/Private/Transport/UdpMessageProcessor.cpp b/Engine/Plugins/Messaging/UdpMessaging/Source/UdpMessaging/Private/Transport/UdpMessageProcessor.cpp
--- a/Engine/Plugins/Messaging/UdpMessaging/Source/UdpMessaging/Private/Transport/UdpMessageProcessor.cpp
+++ b/Engine/Plugins/Messaging/UdpMessaging/Source/UdpMessaging/Private/Transport/UdpMessageProcessor.cpp
@@ -234,9 +234,10 @@ bool FUdpMessageProcessor::Init()
 		Beacon = new FUdpMessageBeacon(Socket, LocalNodeId, MulticastEndpoint);
 		SocketSender = new FUdpSocketSender(Socket, TEXT("FUdpMessageProcessor.Sender"));
 
-		// Current protocol version 14
+		// Current protocol version 15
 		SupportedProtocolVersions.Add(UDP_MESSAGING_TRANSPORT_PROTOCOL_VERSION);
-		// Support Protocol version 10, 11, 12, 13
+		// Support Protocol version 10, 11, 12, 13, 14
+		SupportedProtocolVersions.Add(14);
 		SupportedProtocolVersions.Add(13);
 		SupportedProtocolVersions.Add(12);
 		SupportedProtocolVersions.Add(11);
diff --git a/Engine/Plugins/Messaging/UdpMessaging/Source/UdpMessaging/Private/Transport/UdpSerializeMessageTask.cpp b/Engine/Plugins/Messaging/UdpMessaging/Source/UdpMessaging/Private/Transport/UdpSerializeMessageTask.cpp
--- a/Engine/Plugins/Messaging/UdpMessaging/Source/UdpMessaging/Private/Transport/UdpSerializeMessageTask.cpp
+++ b/Engine/Plugins/Messaging/UdpMessaging/Source/UdpMessaging/Private/Transport/UdpSerializeMessageTask.cpp
@@ -51,8 +51,8 @@ void SerializeMessageV10(FArchive& Archive, const TSharedRef<IMessageContext, ES
 	FStructSerializer::Serialize(MessageContext->GetMessage(), *MessageContext->GetMessageTypeInfo(), Backend);
 }
 
-/** Serialization Routine for message using Protocol version 11, 12, 13 or 14. */
-void SerializeMessageV11_14(FArchive& Archive, const TSharedRef<IMessageContext, ESPMode::ThreadSafe>& MessageContext, EUdpMessageFormat MessageFormat, const EStructSerializerBackendFlags StructSerializerBackendFlags)
+/** Serialization Routine for message using Protocol version 11, 12, 13, 14 or 15. */
+void SerializeMessageV11_15(FArchive& Archive, const TSharedRef<IMessageContext, ESPMode::ThreadSafe>& MessageContext, EUdpMessageFormat MessageFormat, const EStructSerializerBackendFlags StructSerializerBackendFlags)
 {
 	const FName& MessageType = MessageContext->GetMessageType();
 	Archive << const_cast<FName&>(MessageType);
@@ -85,7 +85,7 @@ void SerializeMessageV11_14(FArchive& Archive, const TSharedRef<IMessageContext,
 	}
 
 	// Message Wire Format Id
-	check(MessageFormat == EUdpMessageFormat::CborPlatformEndianness || MessageFormat == EUdpMessageFormat::CborStandardEndianness); // Versions 11 to 14 only supports CBOR.
+	check(MessageFormat == EUdpMessageFormat::CborPlatformEndianness || MessageFormat == EUdpMessageFormat::CborStandardEndianness); // Versions 11 to 15 only supports CBOR.
 	uint8 Format = (uint8)(MessageFormat);
 	Archive << Format;
 	
@@ -122,22 +122,27 @@ void FUdpSerializeMessageTask::DoTask(ENamedThreads::Type CurrentThread, const F
 
 			case 11:
 				ProtocolMaxSegmentSize = UDP_MESSAGING_SEGMENT_SIZE * (int64)UINT16_MAX;
-				UdpSerializeMessageTaskDetails::SerializeMessageV11_14(Archive, MessageContext, EUdpMessageFormat::CborPlatformEndianness, EStructSerializerBackendFlags::Legacy);
+				UdpSerializeMessageTaskDetails::SerializeMessageV11_15(Archive, MessageContext, EUdpMessageFormat::CborPlatformEndianness, EStructSerializerBackendFlags::Legacy);
 				break;
 
 			case 12:
 				ProtocolMaxSegmentSize = UDP_MESSAGING_SEGMENT_SIZE * (int64)INT32_MAX;
-				UdpSerializeMessageTaskDetails::SerializeMessageV11_14(Archive, MessageContext, EUdpMessageFormat::CborPlatformEndianness, EStructSerializerBackendFlags::WriteTextAsComplexString);
+				UdpSerializeMessageTaskDetails::SerializeMessageV11_15(Archive, MessageContext, EUdpMessageFormat::CborPlatformEndianness, EStructSerializerBackendFlags::WriteTextAsComplexString);
 				break;
 
 			case 13:
 				ProtocolMaxSegmentSize = UDP_MESSAGING_SEGMENT_SIZE * (int64)INT32_MAX;
-				UdpSerializeMessageTaskDetails::SerializeMessageV11_14(Archive, MessageContext, EUdpMessageFormat::CborPlatformEndianness, EStructSerializerBackendFlags::Default);
+				UdpSerializeMessageTaskDetails::SerializeMessageV11_15(Archive, MessageContext, EUdpMessageFormat::CborPlatformEndianness, EStructSerializerBackendFlags::LegacyUE4);
 				break;
 
 			case 14:
 				ProtocolMaxSegmentSize = UDP_MESSAGING_SEGMENT_SIZE * (int64)INT32_MAX;
-				UdpSerializeMessageTaskDetails::SerializeMessageV11_14(Archive, MessageContext, SerializedMessage->GetFormat(), EStructSerializerBackendFlags::Default);
+				UdpSerializeMessageTaskDetails::SerializeMessageV11_15(Archive, MessageContext, SerializedMessage->GetFormat(), EStructSerializerBackendFlags::LegacyUE4);
+				break;
+
+			case 15:
+				ProtocolMaxSegmentSize = UDP_MESSAGING_SEGMENT_SIZE * (int64)INT32_MAX;
+				UdpSerializeMessageTaskDetails::SerializeMessageV11_15(Archive, MessageContext, SerializedMessage->GetFormat(), EStructSerializerBackendFlags::Default);
 				break;
 
 			default:
diff --git a/Engine/Plugins/Messaging/UdpMessaging/Source/UdpMessaging/Private/UdpMessagingPrivate.h b/Engine/Plugins/Messaging/UdpMessaging/Source/UdpMessaging/Private/UdpMessagingPrivate.h
--- a/Engine/Plugins/Messaging/UdpMessaging/Source/UdpMessaging/Private/UdpMessagingPrivate.h
+++ b/Engine/Plugins/Messaging/UdpMessaging/Source/UdpMessaging/Private/UdpMessagingPrivate.h
@@ -30,4 +30,4 @@ DECLARE_LOG_CATEGORY_EXTERN(LogUdpMessaging, Log, All);
  * Defines the protocol version of the UDP message transport.
  * @note When changing the version, ensure to update the serialization/deserialization code in UdpSerializeMessageTask.cpp/UdpDeserializedMessage.cpp and the supported versions in FUdpMessageProcessor::Init().
  */
-#define UDP_MESSAGING_TRANSPORT_PROTOCOL_VERSION 14
+#define UDP_MESSAGING_TRANSPORT_PROTOCOL_VERSION 15
diff --git a/Engine/Source/Runtime/Serialization/Private/Backends/CborStructDeserializerBackend.cpp b/Engine/Source/Runtime/Serialization/Private/Backends/CborStructDeserializerBackend.cpp
--- a/Engine/Source/Runtime/Serialization/Private/Backends/CborStructDeserializerBackend.cpp
+++ b/Engine/Source/Runtime/Serialization/Private/Backends/CborStructDeserializerBackend.cpp
@@ -1,14 +1,18 @@
 // Copyright Epic Games, Inc. All Rights Reserved.
 
 #include "Backends/CborStructDeserializerBackend.h"
+
 #include "Backends/StructDeserializerBackendUtilities.h"
+#include "StructSerializationUtilities.h"
 #include "UObject/Class.h"
-#include "UObject/UnrealType.h"
 #include "UObject/EnumProperty.h"
 #include "UObject/TextProperty.h"
+#include "UObject/UnrealType.h"
+
 
-FCborStructDeserializerBackend::FCborStructDeserializerBackend(FArchive& Archive, ECborEndianness CborDataEndianness)
+FCborStructDeserializerBackend::FCborStructDeserializerBackend(FArchive& Archive, ECborEndianness CborDataEndianness, bool bInIsLWCCompatibilityMode)
 	: CborReader(&Archive, CborDataEndianness)
+	, bIsLWCCompatibilityMode(bInIsLWCCompatibilityMode)
 {}
 
 FCborStructDeserializerBackend::~FCborStructDeserializerBackend() = default;
@@ -304,10 +308,21 @@ bool FCborStructDeserializerBackend::ReadProperty(FProperty* Property, FProperty
 			return StructDeserializerBackendUtilities::ClearPropertyValue(Property, Outer, Data, ArrayIndex);
 			// Float
 		case ECborCode::Value_4Bytes:
-			if (FFloatProperty* FloatProperty = CastField<FFloatProperty>(Property))
+
+			if (bIsLWCCompatibilityMode == true && StructSerializationUtilities::IsLWCType(Property->GetOwnerStruct()))
+			{
+				FDoubleProperty* DoubleProperty = CastField<FDoubleProperty>(Property);
+				if (ensureMsgf(DoubleProperty, TEXT("Float field %s with value '%f' from LWC struct type '%s' was expected to be a DoubleProperty but was of type %s (%s)"), *Property->GetFName().ToString(), LastContext.AsFloat(), *Property->GetOwnerStruct()->GetName(), * Property->GetClass()->GetName(), *GetDebugString()))
+				{
+					const double DoubleValue = LastContext.AsFloat();
+					return StructDeserializerBackendUtilities::SetPropertyValue(DoubleProperty, Outer, Data, ArrayIndex, DoubleValue);
+				}
+			}
+			else if(FFloatProperty* FloatProperty = CastField<FFloatProperty>(Property))
 			{
 				return StructDeserializerBackendUtilities::SetPropertyValue(FloatProperty, Outer, Data, ArrayIndex, LastContext.AsFloat());
 			}
+			
 			UE_LOG(LogSerialization, Verbose, TEXT("Float field %s with value '%f' is not supported in FProperty type %s (%s)"), *Property->GetFName().ToString(), LastContext.AsFloat(), *Property->GetClass()->GetName(), *GetDebugString());
 			return false;
 			// Double
diff --git a/Engine/Source/Runtime/Serialization/Private/Backends/CborStructSerializerBackend.cpp b/Engine/Source/Runtime/Serialization/Private/Backends/CborStructSerializerBackend.cpp
--- a/Engine/Source/Runtime/Serialization/Private/Backends/CborStructSerializerBackend.cpp
+++ b/Engine/Source/Runtime/Serialization/Private/Backends/CborStructSerializerBackend.cpp
@@ -1,10 +1,12 @@
 // Copyright Epic Games, Inc. All Rights Reserved.
 
 #include "Backends/CborStructSerializerBackend.h"
-#include "UObject/UnrealType.h"
+
+#include "StructSerializationUtilities.h"
 #include "UObject/EnumProperty.h"
-#include "UObject/TextProperty.h"
 #include "UObject/PropertyPortFlags.h"
+#include "UObject/TextProperty.h"
+#include "UObject/UnrealType.h"
 
 FCborStructSerializerBackend::FCborStructSerializerBackend(FArchive& InArchive)
 	: CborWriter(&InArchive)
@@ -215,7 +217,15 @@ void FCborStructSerializerBackend::WriteProperty(const FStructSerializerState& S
 	// Double & Float
 	else if (State.FieldType == FDoubleProperty::StaticClass())
 	{
-		WritePropertyValue(CborWriter, State, CastFieldChecked<FDoubleProperty>(State.ValueProperty)->GetPropertyValue_InContainer(State.ValueData, ArrayIndex));
+		if (EnumHasAnyFlags(Flags, EStructSerializerBackendFlags::WriteLWCTypesAsFloats) && StructSerializationUtilities::IsLWCType(State.ValueProperty->GetOwnerStruct()))
+		{
+			const double Value = CastFieldChecked<FDoubleProperty>(State.ValueProperty)->GetPropertyValue_InContainer(State.ValueData, ArrayIndex);
+			WritePropertyValue(CborWriter, State, static_cast<float>(Value));
+		}
+		else
+		{
+			WritePropertyValue(CborWriter, State, CastFieldChecked<FDoubleProperty>(State.ValueProperty)->GetPropertyValue_InContainer(State.ValueData, ArrayIndex));
+		}
 	}
 	else if (State.FieldType == FFloatProperty::StaticClass())
 	{
diff --git a/Engine/Source/Runtime/Serialization/Private/StructSerializationUtilities.h b/Engine/Source/Runtime/Serialization/Private/StructSerializationUtilities.h
new file mode 100644
--- /dev/null
+++ b/Engine/Source/Runtime/Serialization/Private/StructSerializationUtilities.h
@@ -0,0 +1,31 @@
+// Copyright Epic Games, Inc. All Rights Reserved.
+
+#pragma once
+
+#include "CoreMinimal.h"
+
+namespace StructSerializationUtilities
+{
+	static bool IsLWCType(const UStruct* Type)
+	{
+		static TSet<FName> LWCTypes = {	NAME_Vector
+										, NAME_Vector2D
+										, NAME_Vector4
+										, NAME_Matrix
+										, NAME_Plane
+										, NAME_Quat
+										, NAME_Rotator
+										, NAME_Transform
+										, NAME_Box
+										, NAME_Box2D
+										, NAME_BoxSphereBounds
+										, FName(TEXT("OrientedBox"))};
+		if (Type)
+		{
+			const FName StructName = Type->GetFName();
+			return LWCTypes.Contains(StructName);
+		}
+
+		return false;
+	}
+};
diff --git a/Engine/Source/Runtime/Serialization/Private/Tests/StructSerializerTest.cpp b/Engine/Source/Runtime/Serialization/Private/Tests/StructSerializerTest.cpp
--- a/Engine/Source/Runtime/Serialization/Private/Tests/StructSerializerTest.cpp
+++ b/Engine/Source/Runtime/Serialization/Private/Tests/StructSerializerTest.cpp
@@ -140,6 +140,178 @@ namespace StructSerializerTest
 		Test.TestTrue(TEXT("Sets.StructSet must be the same before and after de-/serialization"), Struct1.StructSet.Num() == Struct2.StructSet.Num() && Struct1.StructSet.Difference(Struct2.StructSet).Num() == 0);
 	}
 
+	void ValidateLWCSerializationBackwardCompatibility(FAutomationTestBase& Test, const FStructSerializerLWCTypesTest& Struct1, const FStructSerializerNonLWCTypesTest& Struct2)
+	{
+		//Make comparison by casting the double (lwc) version down to float (non-lwc) since this is what will happen during serialization
+		Test.TestEqual<FVector3f>(TEXT("LWC Vector must be deserialized to a Non-LWC Vector"), (FVector3f)Struct1.Vector, Struct2.Vector);
+		Test.TestEqual<FVector2f>(TEXT("LWC Vector2D must be deserialized to a Non-LWC Vector2D"), (FVector2f)Struct1.Vector2D, Struct2.Vector2D);
+		Test.TestEqual<FVector4f>(TEXT("LWC Vector4 must be deserialized to a Non-LWC Vector4"), (FVector4f)Struct1.Vector4, Struct2.Vector4);
+		Test.TestEqual<FMatrix44f>(TEXT("LWC Matrix must be deserialized to a Non-LWC Matrix"), (FMatrix44f)Struct1.Matrix, Struct2.Matrix);
+		Test.TestEqual<FPlane4f>(TEXT("LWC Plane must be deserialized to a Non-LWC Plane"), (FPlane4f)Struct1.Plane, Struct2.Plane);
+		Test.TestEqual<FQuat4f>(TEXT("LWC Quat must be deserialized to a Non-LWC Quat"), (FQuat4f)Struct1.Quat, Struct2.Quat);
+		Test.TestEqual<FRotator3f>(TEXT("LWC Rotator must be deserialized to a Non-LWC Rotator"), (FRotator3f)Struct1.Rotator, Struct2.Rotator);
+		Test.TestTrue(TEXT("LWC Transform must be deserialized to a Non-LWC Transform"), Struct2.Transform.Equals((FTransform3f)Struct1.Transform));
+		Test.TestEqual<FBox3f>(TEXT("LWC Box must be deserialized to a Non-LWC Box"), (FBox3f)Struct1.Box, Struct2.Box);
+		Test.TestEqual<FBox2f>(TEXT("LWC Box2D must be deserialized to a Non-LWC Box2D"), (FBox2f)Struct1.Box2D, Struct2.Box2D);
+		Test.TestEqual<FBoxSphereBounds3f>(TEXT("LWC BoxSphereBounds must be deserialized to a Non-LWC BoxSphereBounds"), (FBoxSphereBounds3f)Struct1.BoxSphereBounds, Struct2.BoxSphereBounds);
+		Test.TestEqual<float>(TEXT("LWC struct float must be the same when deserialized to a non-LWC struct float"), Struct1.Float, Struct2.Float);
+		Test.TestEqual<double>(TEXT("LWC struct double must be the same when deserialized to a non-LWC struct double"), Struct1.Double, Struct2.Double);
+
+		const bool bIsOrientedBoxEqual = Struct1.OrientedBox.AxisX == Struct2.OrientedBox.AxisX
+			&& Struct1.OrientedBox.AxisY == Struct2.OrientedBox.AxisY
+			&& Struct1.OrientedBox.AxisZ == Struct2.OrientedBox.AxisZ
+			&& Struct1.OrientedBox.Center == Struct2.OrientedBox.Center
+			&& Struct1.OrientedBox.ExtentX == Struct2.OrientedBox.ExtentX
+			&& Struct1.OrientedBox.ExtentY == Struct2.OrientedBox.ExtentY
+			&& Struct1.OrientedBox.ExtentZ == Struct2.OrientedBox.ExtentZ;
+		Test.TestTrue(TEXT("LWC OrientedBox must be deserialized to a Non-LWC OrientedBox"), bIsOrientedBoxEqual);
+
+		//Container testing
+		bool bAreArrayEqual = Struct1.VectorArray.Num() == Struct2.VectorArray.Num();
+		for (int32 Index = 0; Index < Struct1.VectorArray.Num() && bAreArrayEqual; ++Index)
+		{
+			if ((FVector3f)Struct1.VectorArray[Index] != Struct2.VectorArray[Index])
+			{
+				bAreArrayEqual = false;
+			}
+		}
+		Test.TestTrue(TEXT("Array of LWC Vectors must be deserialized to an Array of Non-LWC Vectors"), bAreArrayEqual);
+
+		bool bAreMapsEqual = Struct1.StrToVec.Num() == Struct2.StrToVec.Num();
+		for (const TPair<FString, FVector>& Struct1Pair : Struct1.StrToVec)
+		{
+			const FVector3f* Struct2Value = Struct2.StrToVec.Find(Struct1Pair.Key);
+			if (Struct2Value == nullptr)
+			{
+				bAreMapsEqual = false;
+				break;
+			}
+
+			const FVector3f Value = *Struct2Value;
+			if (Value != (FVector3f)Struct1Pair.Value)
+			{
+				bAreMapsEqual = false;
+				break;
+			}
+		}
+		Test.TestTrue(TEXT("Map of LWC Vectors must be deserialized to a Map of Non-LWC Vectors"), bAreMapsEqual);
+
+		bool bAreSetsEqual = Struct1.VectorSet.Num() == Struct2.VectorSet.Num();
+		for (const FVector& Vector : Struct1.VectorSet)
+		{
+			//Cast down like serialization has done
+			const FVector3f FloatVector = (FVector3f)Vector;
+			if (Struct2.VectorSet.Contains(FloatVector) == false)
+			{
+				bAreSetsEqual = false;
+				break;
+			}
+		}
+		Test.TestTrue(TEXT("Set of LWC Vectors must be deserialized to a Set of Non-LWC Vectors"), bAreSetsEqual);
+	}
+
+	void ValidateLWCDeserializationBackwardCompatibility(FAutomationTestBase& Test, const FStructSerializerNonLWCTypesTest& Struct1, const FStructSerializerLWCTypesTest& Struct2)
+	{
+		//Make comparison by casting the float (non-lwc) version up to double (lwc) since this is what will happen during deserialization
+		Test.TestEqual<FVector>(TEXT("Non-LWC Vector must be deserialized to a LWC Vector"), (FVector)Struct1.Vector, Struct2.Vector);
+		Test.TestEqual<FVector2D>(TEXT("Non-LWC Vector2D must be deserialized to a LWC Vector2D"), (FVector2D)Struct1.Vector2D, Struct2.Vector2D);
+		Test.TestEqual<FVector4>(TEXT("Non-LWC Vector4 must be deserialized toa LWC Vector4"), (FVector4)Struct1.Vector4, Struct2.Vector4);
+		Test.TestEqual<FMatrix>(TEXT("Non-LWC Matrix must be deserialized to a LWC Matrix"), (FMatrix)Struct1.Matrix, Struct2.Matrix);
+		Test.TestEqual<FPlane>(TEXT("Non-LWC Plane must be deserialized to a LWC Plane"), (FPlane)Struct1.Plane, Struct2.Plane);
+		Test.TestEqual<FQuat>(TEXT("Non-LWC Quat must be deserialized to a LWC Quat"), (FQuat)Struct1.Quat, Struct2.Quat);
+		Test.TestEqual<FRotator>(TEXT("Non-LWC Rotator must be deserialized to a LWC Rotator"), (FRotator)Struct1.Rotator, Struct2.Rotator);
+		Test.TestTrue(TEXT("Non-LWC Transform must be deserialized to a LWC Transform"), Struct2.Transform.Equals((FTransform)Struct1.Transform));
+		Test.TestEqual<FBox>(TEXT("Non-LWC Box must be deserialized to a LWC Box"), (FBox)Struct1.Box, Struct2.Box);
+
+		Test.TestEqual<FBox2D>(TEXT("Non-LWC Box2D must be deserialized to a LWC Box2D"), (FBox2D)Struct1.Box2D, Struct2.Box2D);
+		Test.TestEqual<FBoxSphereBounds>(TEXT("Non-LWC BoxSphereBounds must be deserialized to a LWC BoxSphereBounds"), (FBoxSphereBounds)Struct1.BoxSphereBounds, Struct2.BoxSphereBounds);
+		Test.TestEqual<float>(TEXT("Non-LWC struct float must be the same when deserialized to a LWC struct float"), Struct1.Float, Struct2.Float);
+		Test.TestEqual<double>(TEXT("Non-LWC struct double must be the same when deserialized to a LWC struct double"), Struct1.Double, Struct2.Double);
+
+		const bool bIsOrientedBoxEqual = Struct1.OrientedBox.AxisX == Struct2.OrientedBox.AxisX
+			&& Struct1.OrientedBox.AxisY == Struct2.OrientedBox.AxisY
+			&& Struct1.OrientedBox.AxisZ == Struct2.OrientedBox.AxisZ
+			&& Struct1.OrientedBox.Center == Struct2.OrientedBox.Center
+			&& Struct1.OrientedBox.ExtentX == Struct2.OrientedBox.ExtentX
+			&& Struct1.OrientedBox.ExtentY == Struct2.OrientedBox.ExtentY
+			&& Struct1.OrientedBox.ExtentZ == Struct2.OrientedBox.ExtentZ;
+		Test.TestTrue(TEXT("Non-LWC OrientedBox must be deserialized to a LWC OrientedBox"), bIsOrientedBoxEqual);
+
+		//Container testing
+		bool bAreArrayEqual = Struct1.VectorArray.Num() == Struct2.VectorArray.Num();
+		for (int32 Index = 0; Index < Struct1.VectorArray.Num() && bAreArrayEqual; ++Index)
+		{
+			if ((FVector)Struct1.VectorArray[Index] != Struct2.VectorArray[Index])
+			{
+				bAreArrayEqual = false;
+			}
+		}
+		Test.TestTrue(TEXT("Array of Non-LWC Vectors must be deserialized to an Array of LWC Vectors"), bAreArrayEqual);
+
+		bool bAreMapsEqual = Struct1.StrToVec.Num() == Struct2.StrToVec.Num();
+		for (const TPair<FString, FVector3f>& Struct1Pair : Struct1.StrToVec)
+		{
+			const FVector* Struct2Value = Struct2.StrToVec.Find(Struct1Pair.Key);
+			if (Struct2Value == nullptr)
+			{
+				bAreMapsEqual = false;
+				break;
+			}
+
+			const FVector Value = *Struct2Value;
+			if (Value != (FVector)Struct1Pair.Value)
+			{
+				bAreMapsEqual = false;
+				break;
+			}
+		}
+		Test.TestTrue(TEXT("Map of Non-LWC Vectors must be deserialized to a Map of LWC Vectors"), bAreMapsEqual);
+
+		bool bAreSetsEqual = Struct1.VectorSet.Num() == Struct2.VectorSet.Num();
+		for (const FVector3f& Vector : Struct1.VectorSet)
+		{
+			//Cast down like serialization has done
+			const FVector FloatVector = (FVector)Vector;
+			if (Struct2.VectorSet.Contains(FloatVector) == false)
+			{
+				bAreSetsEqual = false;
+				break;
+			}
+		}
+		Test.TestTrue(TEXT("Set of Non-LWC Vectors must be deserialized to a Set of LWC Vectors"), bAreSetsEqual);
+	}
+
+	void ValidateLWCTypes(FAutomationTestBase& Test, const FStructSerializerLWCTypesTest& Struct1, const FStructSerializerLWCTypesTest& Struct2)
+	{
+		Test.TestEqual<FVector>(TEXT("LWC Vector must be the same before and after de-serialization"), Struct1.Vector, Struct2.Vector);
+		Test.TestEqual<FVector2D>(TEXT("LWC Vector2D must be the same before and after de-serialization"), Struct1.Vector2D, Struct2.Vector2D);
+		Test.TestEqual<FVector4>(TEXT("LWC Vector4 must be the same before and after de-serialization"), Struct1.Vector4, Struct2.Vector4);
+		Test.TestEqual<FMatrix>(TEXT("LWC Matrix must be the same before and after de-serialization"), Struct1.Matrix, Struct2.Matrix);
+		Test.TestEqual<FPlane>(TEXT("LWC Plane must be the same before and after de-serialization"), Struct1.Plane, Struct2.Plane);
+		Test.TestEqual<FQuat>(TEXT("LWC Quat must be the same before and after de-serialization"), Struct1.Quat, Struct2.Quat);
+		Test.TestEqual<FRotator>(TEXT("LWC Rotator must be the same before and after de-serialization"), Struct1.Rotator, Struct2.Rotator);
+		Test.TestTrue(TEXT("LWC Transform must be the same before and after de-serialization"), Struct2.Transform.Equals(Struct1.Transform));
+		Test.TestEqual<FBox>(TEXT("LWC Box must be the same before and after de-serialization"), Struct1.Box, Struct2.Box);
+
+		Test.TestEqual<FBox2D>(TEXT("LWC Box2D must be the same before and after de-serialization"), Struct1.Box2D, Struct2.Box2D);
+		Test.TestEqual<FBoxSphereBounds>(TEXT("LWC BoxSphereBounds must be the same before and after de-serialization"), Struct1.BoxSphereBounds, Struct2.BoxSphereBounds);
+		Test.TestEqual<float>(TEXT("LWC struct float must be the same before and after de-serialization"), Struct1.Float, Struct2.Float);
+		Test.TestEqual<double>(TEXT("LWC struct double must be the same before and after de-serialization"), Struct1.Double, Struct2.Double);
+
+		const bool bIsOrientedBoxEqual = Struct1.OrientedBox.AxisX == Struct2.OrientedBox.AxisX
+			&& Struct1.OrientedBox.AxisY == Struct2.OrientedBox.AxisY
+			&& Struct1.OrientedBox.AxisZ == Struct2.OrientedBox.AxisZ
+			&& Struct1.OrientedBox.Center == Struct2.OrientedBox.Center
+			&& Struct1.OrientedBox.ExtentX == Struct2.OrientedBox.ExtentX
+			&& Struct1.OrientedBox.ExtentY == Struct2.OrientedBox.ExtentY
+			&& Struct1.OrientedBox.ExtentZ == Struct2.OrientedBox.ExtentZ;
+		Test.TestTrue(TEXT("LWC OrientedBox must be the same before and after de-serialization"), bIsOrientedBoxEqual);
+
+		Test.TestEqual<TArray<FVector>>(TEXT("LWC test - Arrays.VectorArray must be the same before and after de-/serialization"), Struct1.VectorArray, Struct2.VectorArray);
+		Test.TestTrue(TEXT("LWC test - Maps.StrToVec must be the same before and after de-/serialization"), Struct1.StrToVec.OrderIndependentCompareEqual(Struct2.StrToVec));
+		Test.TestTrue(TEXT("LWC test - Sets.VectorSet must be the same before and after de-/serialization"), Struct1.VectorSet.Num() == Struct2.VectorSet.Num() && Struct1.VectorSet.Difference(Struct2.VectorSet).Num() == 0);
+	}
+
 	template<typename TSerializerBackend, typename TDeserializerBackend>
 	void TestElementSerialization(FAutomationTestBase& Test)
 	{
@@ -555,6 +727,49 @@ namespace StructSerializerTest
 		ValidateMaps(Test, TestStruct.Maps, TestStruct2.Maps);
 		// test sets
 		ValidateSets(Test, TestStruct.Sets, TestStruct2.Sets);
+
+		//Test LWC types with standard de-serialization
+		ValidateLWCTypes(Test, TestStruct.LWCTypes, TestStruct2.LWCTypes);
+	}
+
+	void TestLWCSerialization(FAutomationTestBase& Test, IStructSerializerBackend& SerializerBackend, IStructDeserializerBackend& DeserializerBackend)
+	{
+		// Serialization of LWC struct into non-LWC mode to mimick sending to older UE
+		FStructSerializerLWCTypesTest TestLWCStruct;
+		{
+			FStructSerializer::Serialize(TestLWCStruct, SerializerBackend);
+		}
+
+		// Deserialization into non-LWC to mimick reception in an older UE
+		FStructSerializerNonLWCTypesTest TestNonLWCStruct(NoInit);
+		{
+			FStructDeserializerPolicies Policies;
+			Policies.MissingFields = EStructDeserializerErrorPolicies::Warning;
+
+			Test.TestTrue(TEXT("Deserialization must succeed"), FStructDeserializer::Deserialize(TestNonLWCStruct, DeserializerBackend, Policies));
+		}
+
+		ValidateLWCSerializationBackwardCompatibility(Test, TestLWCStruct, TestNonLWCStruct);
+	}
+
+	void TestLWCDeserialization(FAutomationTestBase& Test, IStructSerializerBackend& SerializerBackend, IStructDeserializerBackend& DeserializerBackend)
+	{
+		// Serialization of non lwc struct to mimick a struct coming from older UE
+		FStructSerializerNonLWCTypesTest TestNonLWCStruct;
+		{
+			FStructSerializer::Serialize(TestNonLWCStruct, SerializerBackend);
+		}
+
+		// Deserialization into LWC type to mimick reception into a newer UE
+		FStructSerializerLWCTypesTest TestLWCStruct(NoInit);
+		{
+			FStructDeserializerPolicies Policies;
+			Policies.MissingFields = EStructDeserializerErrorPolicies::Warning;
+
+			Test.TestTrue(TEXT("Deserialization must succeed"), FStructDeserializer::Deserialize(TestLWCStruct, DeserializerBackend, Policies));
+		}
+
+		ValidateLWCDeserializationBackwardCompatibility(Test, TestNonLWCStruct, TestLWCStruct);
 	}
 }
 
@@ -605,6 +820,32 @@ bool FStructSerializerTest::RunTest( const FString& Parameters )
 
 		StructSerializerTest::TestSerialization(*this, SerializerBackend, DeserializerBackend);
 	}
+	// cbor LWC (UE5) to NonLWC (UE4)
+	{
+		TArray<uint8> Buffer;
+		FMemoryReader Reader(Buffer);
+		FMemoryWriter Writer(Buffer);
+
+		FCborStructSerializerBackend SerializerBackend(Writer, EStructSerializerBackendFlags::LegacyUE4 | EStructSerializerBackendFlags::WriteCborStandardEndianness);
+
+		constexpr bool bIsLWCCompatibilityMode = false;
+		FCborStructDeserializerBackend DeserializerBackend(Reader, ECborEndianness::StandardCompliant, bIsLWCCompatibilityMode);
+
+		StructSerializerTest::TestLWCSerialization(*this, SerializerBackend, DeserializerBackend);
+	}
+	// cbor Non LWC (UE4) to LWC (UE5)
+	{
+		TArray<uint8> Buffer;
+		FMemoryReader Reader(Buffer);
+		FMemoryWriter Writer(Buffer);
+
+		FCborStructSerializerBackend SerializerBackend(Writer, EStructSerializerBackendFlags::LegacyUE4 | EStructSerializerBackendFlags::WriteCborStandardEndianness);
+
+		constexpr bool bIsLWCCompatibilityMode = true;
+		FCborStructDeserializerBackend DeserializerBackend(Reader, ECborEndianness::StandardCompliant, bIsLWCCompatibilityMode);
+
+		StructSerializerTest::TestLWCDeserialization(*this, SerializerBackend, DeserializerBackend);
+	}
 
 	return true;
 }
diff --git a/Engine/Source/Runtime/Serialization/Private/Tests/StructSerializerTestTypes.h b/Engine/Source/Runtime/Serialization/Private/Tests/StructSerializerTestTypes.h
--- a/Engine/Source/Runtime/Serialization/Private/Tests/StructSerializerTestTypes.h
+++ b/Engine/Source/Runtime/Serialization/Private/Tests/StructSerializerTestTypes.h
@@ -3,9 +3,12 @@
 #pragma once
 
 #include "CoreMinimal.h"
-#include "UObject/ObjectMacros.h"
+#include "Math/DualQuat.h"
+#include "Math/Sphere.h"
 #include "Misc/Guid.h"
 #include "Templates/SubclassOf.h"
+#include "UObject/ObjectMacros.h"
+
 #include "StructSerializerTestTypes.generated.h"
 
 /**
@@ -244,6 +247,243 @@ struct FStructSerializerBuiltinTestStruct
 	}
 };
 
+/**
+ * Test structure for LWC types.
+ * @see NoExportTypes.h
+ */
+USTRUCT()
+struct FStructSerializerLWCTypesTest
+{
+	GENERATED_BODY()
+
+	UPROPERTY()
+	FVector Vector;
+
+	UPROPERTY()
+	FVector2D Vector2D;
+
+	UPROPERTY()
+	FVector4 Vector4;
+
+	UPROPERTY()
+	FMatrix Matrix;
+
+	UPROPERTY()
+	FPlane Plane;
+	
+	UPROPERTY()
+	FQuat Quat;
+	
+	UPROPERTY()
+	FRotator Rotator;
+
+	UPROPERTY()
+	FTransform Transform;
+
+	UPROPERTY()
+	FBox Box;
+
+	UPROPERTY()
+	FBox2D Box2D;
+
+	UPROPERTY()
+	FBoxSphereBounds BoxSphereBounds;
+	
+	UPROPERTY()
+	FOrientedBox OrientedBox;
+
+	UPROPERTY()
+	float Float;
+
+	UPROPERTY()
+	double Double;
+
+	UPROPERTY()
+	TArray<FVector> VectorArray;
+
+	UPROPERTY()
+	TMap<FString, FVector> StrToVec;
+
+	UPROPERTY()
+	TSet<FVector> VectorSet;
+
+	/** Default constructor. */
+	FStructSerializerLWCTypesTest()
+		: Vector(1.25, 2.5, 3.75)
+		, Vector2D(2.25, 3.5)
+		, Vector4(0.5, 1.25, 2.5, 3.75)
+		, Matrix(FVector(1.25, 2.25, 3.25), FVector(4.25, 5.5, 6.5), FVector(7.25, 8.25, 9.25), FVector(1, 2, 3))
+		, Plane(0.25, 1.25, 2.25, 3.25)
+		, Quat(0.25, 0.5, 0.75, 1.)
+		, Rotator(2.25)
+		, Transform(FQuat(1, 2, 3, 4))
+		, Box(FVector4(1.0, 2.0, 3.0, 4.0), FVector4(5.0, 6.0, 7.0, 8.0))
+		, Box2D(FVector2D(10.0, 20.0), FVector2D(30.0, 40.0))
+		, BoxSphereBounds(FVector(1.0, 2.0, 3.0), FVector(10.0, 10.0, 10.0), 20.0)
+		, Float(5.25f)
+		, Double(1.114)
+	{
+		OrientedBox.AxisX = FVector(1.0, 1.0, 1.0);
+		OrientedBox.AxisY = FVector(2.0, 2.0, 2.0);
+		OrientedBox.AxisZ = FVector(3.0, 3.0, 3.0);
+		OrientedBox.Center = FVector(5.0, 5.0, 5.0);
+		OrientedBox.ExtentX = 10.0;
+		OrientedBox.ExtentY = 20.0;
+		OrientedBox.ExtentZ = 30.0;
+
+		VectorArray.Add(FVector(1.0, 2.0, 3.0));
+		VectorArray.Add(FVector(-1.0, -2.0, -3.0));
+
+		StrToVec.Add(TEXT("V000"), FVector(0.0, 0.0, 0.0));
+		StrToVec.Add(TEXT("V123"), FVector(1.0, 2.0, 3.0));
+		StrToVec.Add(TEXT("V666"), FVector(6.0, 6.0, 6.0));
+
+		VectorSet.Add(FVector(10.0, 11.0, 12.0));
+	}
+
+	/** Creates an uninitialized instance. */
+	FStructSerializerLWCTypesTest(ENoInit) 
+	{
+	
+	}
+};
+
+/** Float (Non LWC) version of FOrientedBox since the float version doesn't exist yet */
+USTRUCT()
+struct FOrientedBoxFloat
+{
+	GENERATED_BODY()
+
+	/** Holds the center of the box. */
+	UPROPERTY()
+	FVector3f Center;
+
+	/** Holds the x-axis vector of the box. Must be a unit vector. */
+	UPROPERTY()
+	FVector3f AxisX;
+
+	/** Holds the y-axis vector of the box. Must be a unit vector. */
+	UPROPERTY()
+	FVector3f AxisY;
+
+	/** Holds the z-axis vector of the box. Must be a unit vector. */
+	UPROPERTY()
+	FVector3f AxisZ;
+
+	/** Holds the extent of the box along its x-axis. */
+	UPROPERTY()
+	float ExtentX;
+
+	/** Holds the extent of the box along its y-axis. */
+	UPROPERTY()
+	float ExtentY;
+
+	/** Holds the extent of the box along its z-axis. */
+	UPROPERTY()
+	float ExtentZ;
+};
+
+/**
+ * Test structure for Non-LWC version of built in types.
+ * @see NoExportTypes.h
+ */
+USTRUCT()
+struct FStructSerializerNonLWCTypesTest
+{
+	GENERATED_BODY()
+
+	UPROPERTY()
+	FVector3f Vector;
+
+	UPROPERTY()
+	FVector2f Vector2D;
+
+	UPROPERTY()
+	FVector4f Vector4;
+
+	UPROPERTY()
+	FMatrix44f Matrix;
+
+	UPROPERTY()
+	FPlane4f Plane;
+
+	UPROPERTY()
+	FQuat4f Quat;
+
+	UPROPERTY()
+	FRotator3f Rotator;
+
+	UPROPERTY()
+	FTransform3f Transform;
+
+	UPROPERTY()
+	FBox3f Box;
+
+	UPROPERTY()
+	FBox2f Box2D;
+
+	UPROPERTY()
+	FBoxSphereBounds3f BoxSphereBounds;
+
+	UPROPERTY()
+	FOrientedBoxFloat OrientedBox;
+
+	UPROPERTY()
+	float Float;
+
+	UPROPERTY()
+	double Double;
+
+	UPROPERTY()
+	TArray<FVector3f> VectorArray;
+
+	UPROPERTY()
+	TMap<FString, FVector3f> StrToVec;
+
+	UPROPERTY()
+	TSet<FVector3f> VectorSet;
+
+	/** Default constructor. */
+	FStructSerializerNonLWCTypesTest()
+		: Vector(1.25f, 2.5f, 3.75f)
+		, Vector2D(2.25f, 3.5f)
+		, Vector4(0.5f, 1.25f, 2.5f, 3.75f)
+		, Matrix(FVector3f(1.25f, 2.25f, 3.25f), FVector3f(4.25f, 5.5f, 6.5f), FVector3f(7.25f, 8.25f, 9.25f), FVector3f(1.f, 2.f, 3.f))
+		, Plane(0.25f, 1.25f, 2.25f, 3.25f)
+		, Quat(0.25f, 0.5f, 0.75f, 1.f)
+		, Rotator(2.25f)
+		, Transform(FQuat4f(1.f, 2.f, 3.f, 4.f))
+		, Box(FVector4f(1.0f, 2.0f, 3.0f, 4.0f), FVector4f(5.0f, 6.0f, 7.0f, 8.0f))
+		, Box2D(FVector2f(10.0f, 20.0f), FVector2f(30.0f, 40.0f))
+		, BoxSphereBounds(FVector3f(1.0f, 2.0f, 3.0f), FVector3f(10.0f, 10.0f, 10.0f), 20.0f)
+		, Float(5.25f)
+		, Double(1.114)
+	{
+		OrientedBox.AxisX = FVector3f(1.0f, 1.0f, 1.0f);
+		OrientedBox.AxisY = FVector3f(2.0f, 2.0f, 2.0f);
+		OrientedBox.AxisZ = FVector3f(3.0f, 3.0f, 3.0f);
+		OrientedBox.Center = FVector3f(5.0f, 5.0f, 5.0f);
+		OrientedBox.ExtentX = 10.0f;
+		OrientedBox.ExtentY = 20.0f;
+		OrientedBox.ExtentZ = 30.0f;
+
+		VectorArray.Add(FVector3f(1.0f, 2.0f, 3.0f));
+		VectorArray.Add(FVector3f(-1.0f, -2.0f, -3.0f));
+
+		StrToVec.Add(TEXT("V000"), FVector3f(0.0f, 0.0f, 0.0f));
+		StrToVec.Add(TEXT("V123"), FVector3f(1.0f, 2.0f, 3.0f));
+		StrToVec.Add(TEXT("V666"), FVector3f(6.0f, 6.0f, 6.0f));
+
+		VectorSet.Add(FVector3f(10.0f, 11.0f, 12.0f));
+	}
+
+	/** Creates an uninitialized instance. */
+	FStructSerializerNonLWCTypesTest(ENoInit)
+	{
+
+	}
+};
+
 // basic type hash to test built in struct in sets
 FORCEINLINE uint32 GetTypeHash(const FStructSerializerBuiltinTestStruct& S)
 {
@@ -470,6 +710,9 @@ struct FStructSerializerTestStruct
 	UPROPERTY()
 	FStructSerializerSetTestStruct Sets;
 
+	UPROPERTY()
+	FStructSerializerLWCTypesTest LWCTypes;
+
 	/** Default constructor. */
 	FStructSerializerTestStruct() = default;
 
@@ -482,5 +725,6 @@ struct FStructSerializerTestStruct
 		, Arrays(NoInit)
 		, Maps(NoInit)
 		, Sets(NoInit)
+		, LWCTypes(NoInit)
 	{ }
 };
diff --git a/Engine/Source/Runtime/Serialization/Public/Backends/CborStructDeserializerBackend.h b/Engine/Source/Runtime/Serialization/Public/Backends/CborStructDeserializerBackend.h
--- a/Engine/Source/Runtime/Serialization/Public/Backends/CborStructDeserializerBackend.h
+++ b/Engine/Source/Runtime/Serialization/Public/Backends/CborStructDeserializerBackend.h
@@ -20,7 +20,7 @@ public:
 	 * @param CborDataEndianness The CBOR data endianness stored in the archive.
 	 * @note For backward compatibility and performance, the implementation default to the the platform endianness rather than the CBOR standard one (big endian).
 	 */
-	FCborStructDeserializerBackend(FArchive& Archive, ECborEndianness CborDataEndianness = ECborEndianness::Platform);
+	FCborStructDeserializerBackend(FArchive& Archive, ECborEndianness CborDataEndianness = ECborEndianness::Platform, bool bInIsLWCCompatibilityMode = false);
 	virtual ~FCborStructDeserializerBackend();
 
 public:
@@ -50,4 +50,7 @@ private:
 
 	/** Whether a TArray<uint8>/TArray<int8> property is being deserialized. */
 	bool bDeserializingByteArray = false;
+
+	/** Whether we are deserializing for LWC backward compability mode. Incoming FloatProperties of LWC types deserialized into Double */
+	bool bIsLWCCompatibilityMode = false;
 };
diff --git a/Engine/Source/Runtime/Serialization/Public/IStructSerializerBackend.h b/Engine/Source/Runtime/Serialization/Public/IStructSerializerBackend.h
--- a/Engine/Source/Runtime/Serialization/Public/IStructSerializerBackend.h
+++ b/Engine/Source/Runtime/Serialization/Public/IStructSerializerBackend.h
@@ -34,13 +34,23 @@ enum class EStructSerializerBackendFlags
 	 */
 	WriteCborStandardEndianness = 1 << 2,
 
+	/**
+	 * Support backward compatibility for LWC types by writing double properties as floats.
+	 */
+	WriteLWCTypesAsFloats = 1 << 3,
+
 	/**
 	 * Legacy settings for backwards compatibility with code compiled prior to 4.22.
 	 */
 	Legacy = None,
 
 	/**
-	 * Default settings for code compiled for 4.25 onwards.
+	 * Legacy settings for backwards compatibility with code compiled for 4.25 up to UE5.
+	 */
+	LegacyUE4 = WriteTextAsComplexString | WriteByteArrayAsByteStream | WriteLWCTypesAsFloats,
+
+	/**
+	 * Default settings for code compiled for 5.0 onwards.
 	 */
 	Default = WriteTextAsComplexString | WriteByteArrayAsByteStream,
 };
