CborStructSerialalizerBackend fixes for FSoftClassProperty. Serialize it like FSoftObjectProperty, where the path can be valid even if the object reference isn't.

diff --git a/Engine/Source/Runtime/Serialization/Private/Backends/CborStructDeserializerBackend.cpp b/Engine/Source/Runtime/Serialization/Private/Backends/CborStructDeserializerBackend.cpp
--- a/Engine/Source/Runtime/Serialization/Private/Backends/CborStructDeserializerBackend.cpp
+++ b/Engine/Source/Runtime/Serialization/Private/Backends/CborStructDeserializerBackend.cpp
@@ -236,7 +236,7 @@ bool FCborStructDeserializerBackend::ReadProperty(FProperty* Property, FProperty
 
 		if (FSoftClassProperty* SoftClassProperty = CastField<FSoftClassProperty>(Property))
 		{
-			return StructDeserializerBackendUtilities::SetPropertyValue(SoftClassProperty, Outer, Data, ArrayIndex, FSoftObjectPtr(LoadObject<UClass>(nullptr, *StringValue, nullptr, LOAD_NoWarn)));
+			return StructDeserializerBackendUtilities::SetPropertyValue(SoftClassProperty, Outer, Data, ArrayIndex, FSoftObjectPtr(FSoftObjectPath(StringValue)));
 		}
 
 		if (FObjectProperty* ObjectProperty = CastField<FObjectProperty>(Property))
diff --git a/Engine/Source/Runtime/Serialization/Private/Backends/CborStructSerializerBackend.cpp b/Engine/Source/Runtime/Serialization/Private/Backends/CborStructSerializerBackend.cpp
--- a/Engine/Source/Runtime/Serialization/Private/Backends/CborStructSerializerBackend.cpp
+++ b/Engine/Source/Runtime/Serialization/Private/Backends/CborStructSerializerBackend.cpp
@@ -300,7 +300,7 @@ void FCborStructSerializerBackend::WriteProperty(const FStructSerializerState& S
 	else if (State.FieldType == FSoftClassProperty::StaticClass())
 	{
 		FSoftObjectPtr const& Value = CastFieldChecked<FSoftClassProperty>(State.ValueProperty)->GetPropertyValue_InContainer(State.ValueData, ArrayIndex);
-		WritePropertyValue(CborWriter, State, Value.IsValid() ? Value->GetPathName() : FString());
+		WritePropertyValue(CborWriter, State, Value.ToString());
 	}
 	else if (State.FieldType == FWeakObjectProperty::StaticClass())
 	{
