Fixes to address FDateTime and FTimespan not serializing via StructSerializer, and fixes to address data loss issues when serializing large integer values via the Json struct serializer backend.

- NoExportTypes.h: Mark FDateTime and FTimespan "Tick" variables as UPROPERTY, to allow serialization via FStructSerializer and FStructDeserializer.
- StructSerializerTestTypes.h: Uncomment FDateTime and FTimespan in FStructSerializerBuiltinTestStruct and initialize them with valid values.
- StructSerializerTest.cpp: Test FDateTime and FTimespan to validate that they are the same value before and after (de)serialization. Change the tests to run in any application (useful for executing tests quickly against a client/server via Session Frontend).
- JsonStructSerializerBackend.cpp: Write numeric property values as their native types, rather than casting to double. This can result in data loss for large uint64/int64 values.
- JsonStructDeserializerBackend.cpp: Read the property value via the number string and parse it as its native type. This can avoid data loss when going through double for large uint64/int64 values.
- JsonWriter.h: Add missing overload for writing uint32 values.

#tests StructSerializer tests pass. Successfully compiled blueprints using DateTime variable (ensuring no regression in UE-44418). Verified that runtime values of FDateTime reflect their default value in blueprint (ensuring no regression in UE-39921).

diff --git a/Engine/Source/Runtime/CoreUObject/Public/UObject/NoExportTypes.h b/Engine/Source/Runtime/CoreUObject/Public/UObject/NoExportTypes.h
--- a/Engine/Source/Runtime/CoreUObject/Public/UObject/NoExportTypes.h
+++ b/Engine/Source/Runtime/CoreUObject/Public/UObject/NoExportTypes.h
@@ -2052,6 +2052,7 @@ public:
 USTRUCT(immutable, noexport, BlueprintType, IsAlwaysAccessible, HasDefaults, meta=(HasNativeMake="/Script/Engine.KismetMathLibrary.MakeDateTime", HasNativeBreak="/Script/Engine.KismetMathLibrary.BreakDateTime"))
 struct FDateTime
 {
+	UPROPERTY()
 	int64 Ticks;
 };
 
@@ -2147,6 +2148,7 @@ struct FTimecode
 USTRUCT(immutable, noexport, BlueprintType, IsAlwaysAccessible, HasDefaults, meta=(HasNativeMake="/Script/Engine.KismetMathLibrary.MakeTimespan", HasNativeBreak="/Script/Engine.KismetMathLibrary.BreakTimespan"))
 struct FTimespan
 {
+	UPROPERTY()
 	int64 Ticks;
 };
 
diff --git a/Engine/Source/Runtime/Json/Public/Serialization/JsonWriter.h b/Engine/Source/Runtime/Json/Public/Serialization/JsonWriter.h
--- a/Engine/Source/Runtime/Json/Public/Serialization/JsonWriter.h
+++ b/Engine/Source/Runtime/Json/Public/Serialization/JsonWriter.h
@@ -437,7 +437,12 @@ protected:
 		return EJsonToken::Number;
 	}
 
-	EJsonToken WriteValueOnly(uint64 Value)
+	FORCEINLINE EJsonToken WriteValueOnly(uint32 Value)
+	{
+		return WriteValueOnly((uint64)Value);
+	}
+
+	FORCEINLINE EJsonToken WriteValueOnly(uint64 Value)
 	{
 		PrintPolicy::WriteString(Stream, WriteToString<32>(Value));
 		return EJsonToken::Number;
diff --git a/Engine/Source/Runtime/Serialization/Private/Backends/JsonStructDeserializerBackend.cpp b/Engine/Source/Runtime/Serialization/Private/Backends/JsonStructDeserializerBackend.cpp
--- a/Engine/Source/Runtime/Serialization/Private/Backends/JsonStructDeserializerBackend.cpp
+++ b/Engine/Source/Runtime/Serialization/Private/Backends/JsonStructDeserializerBackend.cpp
@@ -99,59 +99,99 @@ bool FJsonStructDeserializerBackend::ReadProperty( FProperty* Property, FPropert
 	// numeric values
 	case EJsonNotation::Number:
 		{
-			double NumericValue = JsonReader->GetValueAsNumber();
+			const FString& NumericValueString = JsonReader->GetValueAsNumberString();
 
 			if (FByteProperty* ByteProperty = CastField<FByteProperty>(Property))
 			{
-				return StructDeserializerBackendUtilities::SetPropertyValue(ByteProperty, Outer, Data, ArrayIndex, (int8)NumericValue);
+				int8 NumericValue = 0;
+				if (LexTryParseString(NumericValue, *NumericValueString))
+				{
+					return StructDeserializerBackendUtilities::SetPropertyValue(ByteProperty, Outer, Data, ArrayIndex, NumericValue);
+				}
 			}
 
 			if (FDoubleProperty* DoubleProperty = CastField<FDoubleProperty>(Property))
 			{
-				return StructDeserializerBackendUtilities::SetPropertyValue(DoubleProperty, Outer, Data, ArrayIndex, (double)NumericValue);
+				double NumericValue = 0;
+				if (LexTryParseString(NumericValue, *NumericValueString))
+				{
+					return StructDeserializerBackendUtilities::SetPropertyValue(DoubleProperty, Outer, Data, ArrayIndex, NumericValue);
+				}
 			}
 
 			if (FFloatProperty* FloatProperty = CastField<FFloatProperty>(Property))
 			{
-				return StructDeserializerBackendUtilities::SetPropertyValue(FloatProperty, Outer, Data, ArrayIndex, (float)NumericValue);
+				float NumericValue = 0;
+				if (LexTryParseString(NumericValue, *NumericValueString))
+				{
+					return StructDeserializerBackendUtilities::SetPropertyValue(FloatProperty, Outer, Data, ArrayIndex, NumericValue);
+				}
 			}
 
 			if (FIntProperty* IntProperty = CastField<FIntProperty>(Property))
 			{
-				return StructDeserializerBackendUtilities::SetPropertyValue(IntProperty, Outer, Data, ArrayIndex, (int32)NumericValue);
+				int32 NumericValue = 0;
+				if (LexTryParseString(NumericValue, *NumericValueString))
+				{
+					return StructDeserializerBackendUtilities::SetPropertyValue(IntProperty, Outer, Data, ArrayIndex, NumericValue);
+				}
 			}
 
 			if (FUInt32Property* UInt32Property = CastField<FUInt32Property>(Property))
 			{
-				return StructDeserializerBackendUtilities::SetPropertyValue(UInt32Property, Outer, Data, ArrayIndex, (uint32)NumericValue);
+				uint32 NumericValue = 0;
+				if (LexTryParseString(NumericValue, *NumericValueString))
+				{
+					return StructDeserializerBackendUtilities::SetPropertyValue(UInt32Property, Outer, Data, ArrayIndex, NumericValue);
+				}
 			}
 
 			if (FInt16Property* Int16Property = CastField<FInt16Property>(Property))
 			{
-				return StructDeserializerBackendUtilities::SetPropertyValue(Int16Property, Outer, Data, ArrayIndex, (int16)NumericValue);
+				int16 NumericValue = 0;
+				if (LexTryParseString(NumericValue, *NumericValueString))
+				{
+					return StructDeserializerBackendUtilities::SetPropertyValue(Int16Property, Outer, Data, ArrayIndex, NumericValue);
+				}
 			}
 
 			if (FUInt16Property* FInt16Property = CastField<FUInt16Property>(Property))
 			{
-				return StructDeserializerBackendUtilities::SetPropertyValue(FInt16Property, Outer, Data, ArrayIndex, (uint16)NumericValue);
+				uint16 NumericValue = 0;
+				if (LexTryParseString(NumericValue, *NumericValueString))
+				{
+					return StructDeserializerBackendUtilities::SetPropertyValue(FInt16Property, Outer, Data, ArrayIndex, NumericValue);
+				}
 			}
 
 			if (FInt64Property* Int64Property = CastField<FInt64Property>(Property))
 			{
-				return StructDeserializerBackendUtilities::SetPropertyValue(Int64Property, Outer, Data, ArrayIndex, (int64)NumericValue);
+				int64 NumericValue = 0;
+				if (LexTryParseString(NumericValue, *NumericValueString))
+				{
+					return StructDeserializerBackendUtilities::SetPropertyValue(Int64Property, Outer, Data, ArrayIndex, NumericValue);
+				}
 			}
 
 			if (FUInt64Property* FInt64Property = CastField<FUInt64Property>(Property))
 			{
-				return StructDeserializerBackendUtilities::SetPropertyValue(FInt64Property, Outer, Data, ArrayIndex, (uint64)NumericValue);
+				uint64 NumericValue = 0;
+				if (LexTryParseString(NumericValue, *NumericValueString))
+				{
+					return StructDeserializerBackendUtilities::SetPropertyValue(FInt64Property, Outer, Data, ArrayIndex, NumericValue);
+				}
 			}
 
 			if (FInt8Property* Int8Property = CastField<FInt8Property>(Property))
 			{
-				return StructDeserializerBackendUtilities::SetPropertyValue(Int8Property, Outer, Data, ArrayIndex, (int8)NumericValue);
+				int8 NumericValue = 0;
+				if (LexTryParseString(NumericValue, *NumericValueString))
+				{
+					return StructDeserializerBackendUtilities::SetPropertyValue(Int8Property, Outer, Data, ArrayIndex, NumericValue);
+				}
 			}
 
-			UE_LOG(LogSerialization, Verbose, TEXT("Numeric field %s with value '%f' is not supported in FProperty type %s (%s)"), *Property->GetFName().ToString(), NumericValue, *Property->GetClass()->GetName(), *GetDebugString());
+			UE_LOG(LogSerialization, Verbose, TEXT("Numeric field %s with value '%s' is not supported in FProperty type %s (%s)"), *Property->GetFName().ToString(), *NumericValueString, *Property->GetClass()->GetName(), *GetDebugString());
 
 			return false;
 		}
diff --git a/Engine/Source/Runtime/Serialization/Private/Backends/JsonStructSerializerBackend.cpp b/Engine/Source/Runtime/Serialization/Private/Backends/JsonStructSerializerBackend.cpp
--- a/Engine/Source/Runtime/Serialization/Private/Backends/JsonStructSerializerBackend.cpp
+++ b/Engine/Source/Runtime/Serialization/Private/Backends/JsonStructSerializerBackend.cpp
@@ -102,7 +102,7 @@ void FJsonStructSerializerBackend::WriteProperty(const FStructSerializerState& S
 		}
 		else
 		{
-			WritePropertyValue(State, (double)ByteProperty->GetPropertyValue_InContainer(State.ValueData, ArrayIndex));
+			WritePropertyValue(State, ByteProperty->GetPropertyValue_InContainer(State.ValueData, ArrayIndex));
 		}
 	}
 
@@ -119,33 +119,33 @@ void FJsonStructSerializerBackend::WriteProperty(const FStructSerializerState& S
 	// signed integers
 	else if (State.FieldType == FIntProperty::StaticClass())
 	{
-		WritePropertyValue(State, (double)CastFieldChecked<FIntProperty>(State.ValueProperty)->GetPropertyValue_InContainer(State.ValueData, ArrayIndex));
+		WritePropertyValue(State, CastFieldChecked<FIntProperty>(State.ValueProperty)->GetPropertyValue_InContainer(State.ValueData, ArrayIndex));
 	}
 	else if (State.FieldType == FInt8Property::StaticClass())
 	{
-		WritePropertyValue(State, (double)CastFieldChecked<FInt8Property>(State.ValueProperty)->GetPropertyValue_InContainer(State.ValueData, ArrayIndex));
+		WritePropertyValue(State, CastFieldChecked<FInt8Property>(State.ValueProperty)->GetPropertyValue_InContainer(State.ValueData, ArrayIndex));
 	}
 	else if (State.FieldType == FInt16Property::StaticClass())
 	{
-		WritePropertyValue(State, (double)CastFieldChecked<FInt16Property>(State.ValueProperty)->GetPropertyValue_InContainer(State.ValueData, ArrayIndex));
+		WritePropertyValue(State, CastFieldChecked<FInt16Property>(State.ValueProperty)->GetPropertyValue_InContainer(State.ValueData, ArrayIndex));
 	}
 	else if (State.FieldType == FInt64Property::StaticClass())
 	{
-		WritePropertyValue(State, (double)CastFieldChecked<FInt64Property>(State.ValueProperty)->GetPropertyValue_InContainer(State.ValueData, ArrayIndex));
+		WritePropertyValue(State, CastFieldChecked<FInt64Property>(State.ValueProperty)->GetPropertyValue_InContainer(State.ValueData, ArrayIndex));
 	}
 
 	// unsigned integers
 	else if (State.FieldType == FUInt16Property::StaticClass())
 	{
-		WritePropertyValue(State, (double)CastFieldChecked<FUInt16Property>(State.ValueProperty)->GetPropertyValue_InContainer(State.ValueData, ArrayIndex));
+		WritePropertyValue(State, CastFieldChecked<FUInt16Property>(State.ValueProperty)->GetPropertyValue_InContainer(State.ValueData, ArrayIndex));
 	}
 	else if (State.FieldType == FUInt32Property::StaticClass())
 	{
-		WritePropertyValue(State, (double)CastFieldChecked<FUInt32Property>(State.ValueProperty)->GetPropertyValue_InContainer(State.ValueData, ArrayIndex));
+		WritePropertyValue(State, CastFieldChecked<FUInt32Property>(State.ValueProperty)->GetPropertyValue_InContainer(State.ValueData, ArrayIndex));
 	}
 	else if (State.FieldType == FUInt64Property::StaticClass())
 	{
-		WritePropertyValue(State, (double)CastFieldChecked<FUInt64Property>(State.ValueProperty)->GetPropertyValue_InContainer(State.ValueData, ArrayIndex));
+		WritePropertyValue(State, CastFieldChecked<FUInt64Property>(State.ValueProperty)->GetPropertyValue_InContainer(State.ValueData, ArrayIndex));
 	}
 
 	// names, strings & text
diff --git a/Engine/Source/Runtime/Serialization/Private/Tests/StructSerializerTest.cpp b/Engine/Source/Runtime/Serialization/Private/Tests/StructSerializerTest.cpp
--- a/Engine/Source/Runtime/Serialization/Private/Tests/StructSerializerTest.cpp
+++ b/Engine/Source/Runtime/Serialization/Private/Tests/StructSerializerTest.cpp
@@ -104,6 +104,8 @@ namespace StructSerializerTest
 		Test.TestEqual<FName>(TEXT("Builtins.Name must be the same before and after de-/serialization"), Struct1.Name, Struct2.Name);
 		Test.TestEqual<FString>(TEXT("Builtins.String must be the same before and after de-/serialization"), Struct1.String, Struct2.String);
 		Test.TestEqual<FString>(TEXT("Builtins.Text must be the same before and after de-/serialization"), Struct1.Text.ToString(), Struct2.Text.ToString());
+		Test.TestEqual<FDateTime>(TEXT("Builtins.Datetime must be the same before and after de-/serialization"), Struct1.Datetime, Struct2.Datetime);
+		Test.TestEqual<FTimespan>(TEXT("Builtins.Timespan must be the same before and after de-/serialization"), Struct1.Timespan, Struct2.Timespan);
 		Test.TestEqual<FVector>(TEXT("Builtins.Vector must be the same before and after de-/serialization"), Struct1.Vector, Struct2.Vector);
 		Test.TestEqual<FVector4>(TEXT("Builtins.Vector4 must be the same before and after de-/serialization"), Struct1.Vector4, Struct2.Vector4);
 		Test.TestEqual<FRotator>(TEXT("Builtins.Rotator must be the same before and after de-/serialization"), Struct1.Rotator, Struct2.Rotator);
@@ -779,7 +781,7 @@ namespace StructSerializerTest
 /* Tests
  *****************************************************************************/
 
-IMPLEMENT_SIMPLE_AUTOMATION_TEST(FStructSerializerTest, "System.Core.Serialization.StructSerializer", EAutomationTestFlags::EditorContext | EAutomationTestFlags::EngineFilter)
+IMPLEMENT_SIMPLE_AUTOMATION_TEST(FStructSerializerTest, "System.Core.Serialization.StructSerializer", EAutomationTestFlags::ApplicationContextMask | EAutomationTestFlags::EngineFilter)
 
 
 bool FStructSerializerTest::RunTest( const FString& Parameters )
@@ -852,7 +854,7 @@ bool FStructSerializerTest::RunTest( const FString& Parameters )
 	return true;
 }
 
-IMPLEMENT_SIMPLE_AUTOMATION_TEST(FStructElementSerializerTest, "System.Core.Serialization.StructElementSerializer", EAutomationTestFlags::EditorContext | EAutomationTestFlags::EngineFilter)
+IMPLEMENT_SIMPLE_AUTOMATION_TEST(FStructElementSerializerTest, "System.Core.Serialization.StructElementSerializer", EAutomationTestFlags::ApplicationContextMask | EAutomationTestFlags::EngineFilter)
 
 bool FStructElementSerializerTest::RunTest(const FString& Parameters)
 {
@@ -865,7 +867,7 @@ bool FStructElementSerializerTest::RunTest(const FString& Parameters)
 	return true;
 }
 
-IMPLEMENT_SIMPLE_AUTOMATION_TEST(FStructSerializerCborByteArrayTest, "System.Core.Serialization.StructSerializerCborByteArray", EAutomationTestFlags::EditorContext | EAutomationTestFlags::EngineFilter)
+IMPLEMENT_SIMPLE_AUTOMATION_TEST(FStructSerializerCborByteArrayTest, "System.Core.Serialization.StructSerializerCborByteArray", EAutomationTestFlags::ApplicationContextMask | EAutomationTestFlags::EngineFilter)
 
 bool FStructSerializerCborByteArrayTest::RunTest( const FString& Parameters )
 {
diff --git a/Engine/Source/Runtime/Serialization/Private/Tests/StructSerializerTestTypes.h b/Engine/Source/Runtime/Serialization/Private/Tests/StructSerializerTestTypes.h
--- a/Engine/Source/Runtime/Serialization/Private/Tests/StructSerializerTestTypes.h
+++ b/Engine/Source/Runtime/Serialization/Private/Tests/StructSerializerTestTypes.h
@@ -201,12 +201,11 @@ struct FStructSerializerBuiltinTestStruct
 	UPROPERTY()
 	FText Text;
 
-	// FDatetime and FTimespan should be tested here but aren't properly setup in `NoExportTypes.h` and so do not properly work currently.
-	//UPROPERTY()
-	//FDateTime Datetime;
+	UPROPERTY()
+	FDateTime Datetime;
 
-	//UPROPERTY()
-	//FTimespan Timespan;
+	UPROPERTY()
+	FTimespan Timespan;
 
 	UPROPERTY()
 	FVector Vector;
@@ -229,6 +228,8 @@ struct FStructSerializerBuiltinTestStruct
 		, Name(TEXT("Test FName"))
 		, String("Test String")
 		, Text(FText::FromString("Test Text"))
+		, Datetime(2048, 2, 4, 8, 16, 32, 64)
+		, Timespan(1, 2, 3, 4, 5)
 		, Vector(1.0f, 2.0f, 3.0f)
 		, Vector4(4.0f, 5.0f, 6.0f, 7.0f)
 		, Rotator(4096, 8192, 16384)
@@ -241,7 +242,7 @@ struct FStructSerializerBuiltinTestStruct
 
 	bool operator==(const FStructSerializerBuiltinTestStruct& Rhs) const
 	{
-		return Guid == Rhs.Guid && Name == Rhs.Name && String == Rhs.String && Text.EqualTo(Rhs.Text) && Vector == Rhs.Vector && Vector4 == Rhs.Vector4 && Rotator == Rhs.Rotator && Quat == Rhs.Quat && Color == Rhs.Color;
+		return Guid == Rhs.Guid && Name == Rhs.Name && String == Rhs.String && Text.EqualTo(Rhs.Text) && Datetime == Rhs.Datetime && Timespan == Rhs.Timespan && Vector == Rhs.Vector && Vector4 == Rhs.Vector4 && Rotator == Rhs.Rotator && Quat == Rhs.Quat && Color == Rhs.Color;
 	}
 };
 
