Update StructSerializerTest to LLT, and add Serialization module as a dependency of FoundationTests to include the StructSerializerTest in it.

#tests Ran serialization tests via FoundationTests LLT target and the automation framework

diff --git a/Engine/Source/Programs/LowLevelTests/FoundationTests/FoundationTests.Build.cs b/Engine/Source/Programs/LowLevelTests/FoundationTests/FoundationTests.Build.cs
--- a/Engine/Source/Programs/LowLevelTests/FoundationTests/FoundationTests.Build.cs
+++ b/Engine/Source/Programs/LowLevelTests/FoundationTests/FoundationTests.Build.cs
@@ -31,7 +31,8 @@ public class FoundationTests : TestModuleRules
 				"CoreUObject",
 				"TelemetryUtils",
 				"AssetRegistry",
-				"ColorManagement"
+				"ColorManagement",
+				"Serialization",
 			});
 
 		if (Target.bBuildWithEditorOnlyData)
diff --git a/Engine/Source/Runtime/Serialization/Private/Tests/StructSerializerTest.cpp b/Engine/Source/Runtime/Serialization/Private/Tests/StructSerializerTest.cpp
--- a/Engine/Source/Runtime/Serialization/Private/Tests/StructSerializerTest.cpp
+++ b/Engine/Source/Runtime/Serialization/Private/Tests/StructSerializerTest.cpp
@@ -1,5 +1,7 @@
 // Copyright Epic Games, Inc. All Rights Reserved.
 
+#if WITH_TESTS
+
 #include "Algo/ForEach.h"
 #include "CoreMinimal.h"
 #include "Misc/Guid.h"
@@ -14,10 +16,12 @@
 #include "StructDeserializer.h"
 #include "StructSerializer.h"
 #include "Tests/StructSerializerTestTypes.h"
-
+#include "Tests/TestHarnessAdapter.h"
 #include "UObject/MetaData.h"
 
-#if WITH_DEV_AUTOMATION_TESTS
+#if WITH_LOW_LEVEL_TESTS
+#include "TestCommon/Expectations.h"
+#endif
 
 /* Internal helpers
  *****************************************************************************/
@@ -56,108 +60,108 @@ namespace StructSerializerTest
 
 	
 
-	void ValidateNumerics(FAutomationTestBase& Test, const FStructSerializerNumericTestStruct& Struct1, const FStructSerializerNumericTestStruct& Struct2)
+	void ValidateNumerics(const FStructSerializerNumericTestStruct& Struct1, const FStructSerializerNumericTestStruct& Struct2)
 	{
-		Test.TestEqual<int8>(TEXT("Numerics.Int8 value must be the same before and after de-/serialization"), Struct1.Int8, Struct2.Int8);
-		Test.TestEqual<int16>(TEXT("Numerics.Int16 value must be the same before and after de-/serialization"), Struct1.Int16, Struct2.Int16);
-		Test.TestEqual<int32>(TEXT("Numerics.Int32 value must be the same before and after de-/serialization"), Struct1.Int32, Struct2.Int32);
-		Test.TestEqual<int64>(TEXT("Numerics.Int64 value must be the same before and after de-/serialization"), Struct1.Int64, Struct2.Int64);
-		Test.TestEqual<uint8>(TEXT("Numerics.UInt8 value must be the same before and after de-/serialization"), Struct1.UInt8, Struct2.UInt8);
-		Test.TestEqual<uint16>(TEXT("Numerics.UInt16 value must be the same before and after de-/serialization"), Struct1.UInt16, Struct2.UInt16);
-		Test.TestEqual<uint32>(TEXT("Numerics.UInt32 value must be the same before and after de-/serialization"), Struct1.UInt32, Struct2.UInt32);
-		Test.TestEqual<uint64>(TEXT("Numerics.UInt64 value must be the same before and after de-/serialization"), Struct1.UInt64, Struct2.UInt64);
-		Test.TestEqual<float>(TEXT("Numerics.Float value must be the same before and after de-/serialization"), Struct1.Float, Struct2.Float);
-		Test.TestEqual<double>(TEXT("Numerics.Double value must be the same before and after de-/serialization"), Struct1.Double, Struct2.Double);
+		CHECK_EQUALS(TEXT("Numerics.Int8 value must be the same before and after de-/serialization"), Struct1.Int8, Struct2.Int8);
+		CHECK_EQUALS(TEXT("Numerics.Int16 value must be the same before and after de-/serialization"), Struct1.Int16, Struct2.Int16);
+		CHECK_EQUALS(TEXT("Numerics.Int32 value must be the same before and after de-/serialization"), Struct1.Int32, Struct2.Int32);
+		CHECK_EQUALS(TEXT("Numerics.Int64 value must be the same before and after de-/serialization"), Struct1.Int64, Struct2.Int64);
+		CHECK_EQUALS(TEXT("Numerics.UInt8 value must be the same before and after de-/serialization"), Struct1.UInt8, Struct2.UInt8);
+		CHECK_EQUALS(TEXT("Numerics.UInt16 value must be the same before and after de-/serialization"), Struct1.UInt16, Struct2.UInt16);
+		CHECK_EQUALS(TEXT("Numerics.UInt32 value must be the same before and after de-/serialization"), Struct1.UInt32, Struct2.UInt32);
+		CHECK_EQUALS(TEXT("Numerics.UInt64 value must be the same before and after de-/serialization"), Struct1.UInt64, Struct2.UInt64);
+		CHECK_EQUALS(TEXT("Numerics.Float value must be the same before and after de-/serialization"), Struct1.Float, Struct2.Float);
+		CHECK_EQUALS(TEXT("Numerics.Double value must be the same before and after de-/serialization"), Struct1.Double, Struct2.Double);
 	}
 
-	void ValidateBooleans(FAutomationTestBase& Test, const FStructSerializerBooleanTestStruct& Struct1, const FStructSerializerBooleanTestStruct& Struct2)
+	void ValidateBooleans(const FStructSerializerBooleanTestStruct& Struct1, const FStructSerializerBooleanTestStruct& Struct2)
 	{
-		Test.TestEqual<bool>(TEXT("Booleans.BoolFalse must be the same before and after de-/serialization"), Struct1.BoolFalse, Struct2.BoolFalse);
-		Test.TestEqual<bool>(TEXT("Booleans.BoolTrue must be the same before and after de-/serialization"), Struct1.BoolTrue, Struct2.BoolTrue);
-		Test.TestEqual<bool>(TEXT("Booleans.Bitfield0 must be the same before and after de-/serialization"), Struct1.Bitfield0, Struct2.Bitfield0);
-		Test.TestEqual<bool>(TEXT("Booleans.Bitfield1 must be the same before and after de-/serialization"), Struct1.Bitfield1, Struct2.Bitfield1);
-		Test.TestEqual<bool>(TEXT("Booleans.Bitfield2Set must be the same before and after de-/serialization"), Struct1.Bitfield2Set, Struct2.Bitfield2Set);
-		Test.TestEqual<bool>(TEXT("Booleans.Bitfield3 must be the same before and after de-/serialization"), Struct1.Bitfield3, Struct2.Bitfield3);
-		Test.TestEqual<bool>(TEXT("Booleans.Bitfield4Set must be the same before and after de-/serialization"), Struct1.Bitfield4Set, Struct2.Bitfield4Set);
-		Test.TestEqual<bool>(TEXT("Booleans.Bitfield5Set must be the same before and after de-/serialization"), Struct1.Bitfield5Set, Struct2.Bitfield5Set);
-		Test.TestEqual<bool>(TEXT("Booleans.Bitfield6 must be the same before and after de-/serialization"), Struct1.Bitfield6, Struct2.Bitfield6);
-		Test.TestEqual<bool>(TEXT("Booleans.Bitfield7 must be the same before and after de-/serialization"), Struct1.Bitfield7Set, Struct2.Bitfield7Set);
+		CHECK_EQUALS(TEXT("Booleans.BoolFalse must be the same before and after de-/serialization"), Struct1.BoolFalse, Struct2.BoolFalse);
+		CHECK_EQUALS(TEXT("Booleans.BoolTrue must be the same before and after de-/serialization"), Struct1.BoolTrue, Struct2.BoolTrue);
+		CHECK_EQUALS(TEXT("Booleans.Bitfield0 must be the same before and after de-/serialization"), Struct1.Bitfield0, Struct2.Bitfield0);
+		CHECK_EQUALS(TEXT("Booleans.Bitfield1 must be the same before and after de-/serialization"), Struct1.Bitfield1, Struct2.Bitfield1);
+		CHECK_EQUALS(TEXT("Booleans.Bitfield2Set must be the same before and after de-/serialization"), Struct1.Bitfield2Set, Struct2.Bitfield2Set);
+		CHECK_EQUALS(TEXT("Booleans.Bitfield3 must be the same before and after de-/serialization"), Struct1.Bitfield3, Struct2.Bitfield3);
+		CHECK_EQUALS(TEXT("Booleans.Bitfield4Set must be the same before and after de-/serialization"), Struct1.Bitfield4Set, Struct2.Bitfield4Set);
+		CHECK_EQUALS(TEXT("Booleans.Bitfield5Set must be the same before and after de-/serialization"), Struct1.Bitfield5Set, Struct2.Bitfield5Set);
+		CHECK_EQUALS(TEXT("Booleans.Bitfield6 must be the same before and after de-/serialization"), Struct1.Bitfield6, Struct2.Bitfield6);
+		CHECK_EQUALS(TEXT("Booleans.Bitfield7 must be the same before and after de-/serialization"), Struct1.Bitfield7Set, Struct2.Bitfield7Set);
 	}
 
-	void ValidateObjects(FAutomationTestBase& Test, const FStructSerializerObjectTestStruct& Struct1, const FStructSerializerObjectTestStruct& Struct2)
+	void ValidateObjects(const FStructSerializerObjectTestStruct& Struct1, const FStructSerializerObjectTestStruct& Struct2)
 	{
-		Test.TestEqual<class UClass*>(TEXT("Objects.RawClass must be the same before and after de-/serialization"), Struct1.RawClass, Struct2.RawClass);
-		Test.TestEqual<class UClass*>(TEXT("Objects.WrappedClass must be the same before and after de-/serialization"), Struct1.WrappedClass, Struct2.WrappedClass);
-		Test.TestEqual<TSubclassOf<class UMetaData>>(TEXT("Objects.SubClass must be the same before and after de-/serialization"), Struct1.SubClass, Struct2.SubClass);
-		Test.TestEqual<TSoftClassPtr<class UMetaData>>(TEXT("Objects.SoftClass must be the same before and after de-/serialization"), Struct1.SoftClass, Struct2.SoftClass);
-		Test.TestEqual<UObject*>(TEXT("Objects.RawObject must be the same before and after de-/serialization"), Struct1.RawObject, Struct2.RawObject);
-		Test.TestEqual<UObject*>(TEXT("Objects.WrappedObject must be the same before and after de-/serialization"), Struct1.WrappedObject, Struct2.WrappedObject);
-		Test.TestEqual<TWeakObjectPtr<class UMetaData>>(TEXT("Objects.WeakObject must be the same before and after de-/serialization"), Struct1.WeakObject, Struct2.WeakObject);
-		Test.TestEqual<TSoftObjectPtr<class UMetaData>>(TEXT("Objects.SoftObject must be the same before and after de-/serialization"), Struct1.SoftObject, Struct2.SoftObject);
-		Test.TestEqual<FSoftClassPath>(TEXT("Objects.ClassPath must be the same before and after de-/serialization"), Struct1.ClassPath, Struct2.ClassPath);
-		Test.TestEqual<FSoftObjectPath>(TEXT("Objects.ObjectPath must be the same before and after de-/serialization"), Struct1.ObjectPath, Struct2.ObjectPath);
+		CHECK_EQUALS(TEXT("Objects.RawClass must be the same before and after de-/serialization"), Struct1.RawClass, Struct2.RawClass);
+		CHECK_EQUALS(TEXT("Objects.WrappedClass must be the same before and after de-/serialization"), Struct1.WrappedClass, Struct2.WrappedClass);
+		CHECK_EQUALS(TEXT("Objects.SubClass must be the same before and after de-/serialization"), Struct1.SubClass, Struct2.SubClass);
+		CHECK_EQUALS(TEXT("Objects.SoftClass must be the same before and after de-/serialization"), Struct1.SoftClass, Struct2.SoftClass);
+		CHECK_EQUALS(TEXT("Objects.RawObject must be the same before and after de-/serialization"), Struct1.RawObject, Struct2.RawObject);
+		CHECK_EQUALS(TEXT("Objects.WrappedObject must be the same before and after de-/serialization"), Struct1.WrappedObject, Struct2.WrappedObject);
+		CHECK_EQUALS(TEXT("Objects.WeakObject must be the same before and after de-/serialization"), Struct1.WeakObject, Struct2.WeakObject);
+		CHECK_EQUALS(TEXT("Objects.SoftObject must be the same before and after de-/serialization"), Struct1.SoftObject, Struct2.SoftObject);
+		CHECK_EQUALS(TEXT("Objects.ClassPath must be the same before and after de-/serialization"), Struct1.ClassPath, Struct2.ClassPath);
+		CHECK_EQUALS(TEXT("Objects.ObjectPath must be the same before and after de-/serialization"), Struct1.ObjectPath, Struct2.ObjectPath);
 	}
 
-	void ValidateBuiltIns(FAutomationTestBase& Test, const FStructSerializerBuiltinTestStruct& Struct1, const FStructSerializerBuiltinTestStruct& Struct2)
+	void ValidateBuiltIns(const FStructSerializerBuiltinTestStruct& Struct1, const FStructSerializerBuiltinTestStruct& Struct2)
 	{
-		Test.TestEqual<FGuid>(TEXT("Builtins.Guid must be the same before and after de-/serialization"), Struct1.Guid, Struct2.Guid);
-		Test.TestEqual<FName>(TEXT("Builtins.Name must be the same before and after de-/serialization"), Struct1.Name, Struct2.Name);
-		Test.TestEqual<FString>(TEXT("Builtins.String must be the same before and after de-/serialization"), Struct1.String, Struct2.String);
-		Test.TestEqual<FString>(TEXT("Builtins.Text must be the same before and after de-/serialization"), Struct1.Text.ToString(), Struct2.Text.ToString());
-		Test.TestEqual<FDateTime>(TEXT("Builtins.Datetime must be the same before and after de-/serialization"), Struct1.Datetime, Struct2.Datetime);
-		Test.TestEqual<FTimespan>(TEXT("Builtins.Timespan must be the same before and after de-/serialization"), Struct1.Timespan, Struct2.Timespan);
-		Test.TestEqual<FVector>(TEXT("Builtins.Vector must be the same before and after de-/serialization"), Struct1.Vector, Struct2.Vector);
-		Test.TestEqual<FVector4>(TEXT("Builtins.Vector4 must be the same before and after de-/serialization"), Struct1.Vector4, Struct2.Vector4);
-		Test.TestEqual<FRotator>(TEXT("Builtins.Rotator must be the same before and after de-/serialization"), Struct1.Rotator, Struct2.Rotator);
-		Test.TestEqual<FQuat>(TEXT("Builtins.Quat must be the same before and after de-/serialization"), Struct1.Quat, Struct2.Quat);
-		Test.TestEqual<FColor>(TEXT("Builtins.Color must be the same before and after de-/serialization"), Struct1.Color, Struct2.Color);
+		CHECK_EQUALS(TEXT("Builtins.Guid must be the same before and after de-/serialization"), Struct1.Guid, Struct2.Guid);
+		CHECK_EQUALS(TEXT("Builtins.Name must be the same before and after de-/serialization"), Struct1.Name, Struct2.Name);
+		CHECK_EQUALS(TEXT("Builtins.String must be the same before and after de-/serialization"), Struct1.String, Struct2.String);
+		CHECK_EQUALS(TEXT("Builtins.Text must be the same before and after de-/serialization"), Struct1.Text.ToString(), Struct2.Text.ToString());
+		CHECK_EQUALS(TEXT("Builtins.Datetime must be the same before and after de-/serialization"), Struct1.Datetime, Struct2.Datetime);
+		CHECK_EQUALS(TEXT("Builtins.Timespan must be the same before and after de-/serialization"), Struct1.Timespan, Struct2.Timespan);
+		CHECK_EQUALS(TEXT("Builtins.Vector must be the same before and after de-/serialization"), Struct1.Vector, Struct2.Vector);
+		CHECK_EQUALS(TEXT("Builtins.Vector4 must be the same before and after de-/serialization"), Struct1.Vector4, Struct2.Vector4);
+		CHECK_EQUALS(TEXT("Builtins.Rotator must be the same before and after de-/serialization"), Struct1.Rotator, Struct2.Rotator);
+		CHECK_EQUALS(TEXT("Builtins.Quat must be the same before and after de-/serialization"), Struct1.Quat, Struct2.Quat);
+		CHECK_EQUALS(TEXT("Builtins.Color must be the same before and after de-/serialization"), Struct1.Color, Struct2.Color);
 	}
 
-	void ValidateArrays(FAutomationTestBase& Test, const FStructSerializerArrayTestStruct& Struct1, const FStructSerializerArrayTestStruct& Struct2)
+	void ValidateArrays(const FStructSerializerArrayTestStruct& Struct1, const FStructSerializerArrayTestStruct& Struct2)
 	{
-		Test.TestEqual<TArray<int32>>(TEXT("Arrays.Int32Array must be the same before and after de-/serialization"), Struct1.Int32Array, Struct2.Int32Array);
-		Test.TestEqual<TArray<uint8>>(TEXT("Arrays.ByteArray must be the same before and after de-/serialization"), Struct1.ByteArray, Struct2.ByteArray);
-		Test.TestEqual<int32>(TEXT("Arrays.StaticSingleElement[0] must be the same before and after de-/serialization"), Struct1.StaticSingleElement[0], Struct2.StaticSingleElement[0]);
-		Test.TestEqual<int32>(TEXT("Arrays.StaticInt32Array[0] must be the same before and after de-/serialization"), Struct1.StaticInt32Array[0], Struct2.StaticInt32Array[0]);
-		Test.TestEqual<int32>(TEXT("Arrays.StaticInt32Array[1] must be the same before and after de-/serialization"), Struct1.StaticInt32Array[1], Struct2.StaticInt32Array[1]);
-		Test.TestEqual<int32>(TEXT("Arrays.StaticInt32Array[2] must be the same before and after de-/serialization"), Struct1.StaticInt32Array[2], Struct2.StaticInt32Array[2]);
-		Test.TestEqual<float>(TEXT("Arrays.StaticFloatArray[0] must be the same before and after de-/serialization"), Struct1.StaticFloatArray[0], Struct2.StaticFloatArray[0]);
-		Test.TestEqual<float>(TEXT("Arrays.StaticFloatArray[1] must be the same before and after de-/serialization"), Struct1.StaticFloatArray[1], Struct2.StaticFloatArray[1]);
-		Test.TestEqual<float>(TEXT("Arrays.StaticFloatArray[2] must be the same before and after de-/serialization"), Struct1.StaticFloatArray[2], Struct2.StaticFloatArray[2]);
-		Test.TestEqual<TArray<FVector>>(TEXT("Arrays.VectorArray must be the same before and after de-/serialization"), Struct1.VectorArray, Struct2.VectorArray);
+		CHECK_EQUALS(TEXT("Arrays.Int32Array must be the same before and after de-/serialization"), Struct1.Int32Array, Struct2.Int32Array);
+		CHECK_EQUALS(TEXT("Arrays.ByteArray must be the same before and after de-/serialization"), Struct1.ByteArray, Struct2.ByteArray);
+		CHECK_EQUALS(TEXT("Arrays.StaticSingleElement[0] must be the same before and after de-/serialization"), Struct1.StaticSingleElement[0], Struct2.StaticSingleElement[0]);
+		CHECK_EQUALS(TEXT("Arrays.StaticInt32Array[0] must be the same before and after de-/serialization"), Struct1.StaticInt32Array[0], Struct2.StaticInt32Array[0]);
+		CHECK_EQUALS(TEXT("Arrays.StaticInt32Array[1] must be the same before and after de-/serialization"), Struct1.StaticInt32Array[1], Struct2.StaticInt32Array[1]);
+		CHECK_EQUALS(TEXT("Arrays.StaticInt32Array[2] must be the same before and after de-/serialization"), Struct1.StaticInt32Array[2], Struct2.StaticInt32Array[2]);
+		CHECK_EQUALS(TEXT("Arrays.StaticFloatArray[0] must be the same before and after de-/serialization"), Struct1.StaticFloatArray[0], Struct2.StaticFloatArray[0]);
+		CHECK_EQUALS(TEXT("Arrays.StaticFloatArray[1] must be the same before and after de-/serialization"), Struct1.StaticFloatArray[1], Struct2.StaticFloatArray[1]);
+		CHECK_EQUALS(TEXT("Arrays.StaticFloatArray[2] must be the same before and after de-/serialization"), Struct1.StaticFloatArray[2], Struct2.StaticFloatArray[2]);
+		CHECK_EQUALS(TEXT("Arrays.VectorArray must be the same before and after de-/serialization"), Struct1.VectorArray, Struct2.VectorArray);
 	}
 
-	void ValidateMaps(FAutomationTestBase& Test, const FStructSerializerMapTestStruct& Struct1, const FStructSerializerMapTestStruct& Struct2)
+	void ValidateMaps(const FStructSerializerMapTestStruct& Struct1, const FStructSerializerMapTestStruct& Struct2)
 	{
-		Test.TestTrue(TEXT("Maps.IntToStr must be the same before and after de-/serialization"), Struct1.IntToStr.OrderIndependentCompareEqual(Struct2.IntToStr));
-		Test.TestTrue(TEXT("Maps.StrToStr must be the same before and after de-/serialization"), Struct1.StrToStr.OrderIndependentCompareEqual(Struct2.StrToStr));
-		Test.TestTrue(TEXT("Maps.StrToVec must be the same before and after de-/serialization"), Struct1.StrToVec.OrderIndependentCompareEqual(Struct2.StrToVec));
+		CHECK_MESSAGE(TEXT("Maps.IntToStr must be the same before and after de-/serialization"), Struct1.IntToStr.OrderIndependentCompareEqual(Struct2.IntToStr));
+		CHECK_MESSAGE(TEXT("Maps.StrToStr must be the same before and after de-/serialization"), Struct1.StrToStr.OrderIndependentCompareEqual(Struct2.StrToStr));
+		CHECK_MESSAGE(TEXT("Maps.StrToVec must be the same before and after de-/serialization"), Struct1.StrToVec.OrderIndependentCompareEqual(Struct2.StrToVec));
 	}
 
-	void ValidateSets(FAutomationTestBase& Test, const FStructSerializerSetTestStruct& Struct1, const FStructSerializerSetTestStruct& Struct2)
+	void ValidateSets(const FStructSerializerSetTestStruct& Struct1, const FStructSerializerSetTestStruct& Struct2)
 	{
-		Test.TestTrue(TEXT("Sets.IntSet must be the same before and after de-/serialization"), Struct1.IntSet.Num() == Struct2.IntSet.Num() && Struct1.IntSet.Difference(Struct2.IntSet).Num() == 0);
-		Test.TestTrue(TEXT("Sets.StrSet must be the same before and after de-/serialization"), Struct1.StrSet.Num() == Struct2.StrSet.Num() && Struct1.StrSet.Difference(Struct2.StrSet).Num() == 0);
-		Test.TestTrue(TEXT("Sets.NameSet must be the same before and after de-/serialization"), Struct1.NameSet.Num() == Struct2.NameSet.Num() && Struct1.NameSet.Difference(Struct2.NameSet).Num() == 0);
-		Test.TestTrue(TEXT("Sets.StructSet must be the same before and after de-/serialization"), Struct1.StructSet.Num() == Struct2.StructSet.Num() && Struct1.StructSet.Difference(Struct2.StructSet).Num() == 0);
+		CHECK_MESSAGE(TEXT("Sets.IntSet must be the same before and after de-/serialization"), Struct1.IntSet.Num() == Struct2.IntSet.Num() && Struct1.IntSet.Difference(Struct2.IntSet).Num() == 0);
+		CHECK_MESSAGE(TEXT("Sets.StrSet must be the same before and after de-/serialization"), Struct1.StrSet.Num() == Struct2.StrSet.Num() && Struct1.StrSet.Difference(Struct2.StrSet).Num() == 0);
+		CHECK_MESSAGE(TEXT("Sets.NameSet must be the same before and after de-/serialization"), Struct1.NameSet.Num() == Struct2.NameSet.Num() && Struct1.NameSet.Difference(Struct2.NameSet).Num() == 0);
+		CHECK_MESSAGE(TEXT("Sets.StructSet must be the same before and after de-/serialization"), Struct1.StructSet.Num() == Struct2.StructSet.Num() && Struct1.StructSet.Difference(Struct2.StructSet).Num() == 0);
 	}
 
-	void ValidateLWCSerializationBackwardCompatibility(FAutomationTestBase& Test, const FStructSerializerLWCTypesTest& Struct1, const FStructSerializerNonLWCTypesTest& Struct2)
+	void ValidateLWCSerializationBackwardCompatibility(const FStructSerializerLWCTypesTest& Struct1, const FStructSerializerNonLWCTypesTest& Struct2)
 	{
 		//Make comparison by casting the double (lwc) version down to float (non-lwc) since this is what will happen during serialization
-		Test.TestEqual<FVector3f>(TEXT("LWC Vector must be deserialized to a Non-LWC Vector"), (FVector3f)Struct1.Vector, Struct2.Vector);
-		Test.TestEqual<FVector2f>(TEXT("LWC Vector2D must be deserialized to a Non-LWC Vector2D"), (FVector2f)Struct1.Vector2D, Struct2.Vector2D);
-		Test.TestEqual<FVector4f>(TEXT("LWC Vector4 must be deserialized to a Non-LWC Vector4"), (FVector4f)Struct1.Vector4, Struct2.Vector4);
-		Test.TestEqual<FMatrix44f>(TEXT("LWC Matrix must be deserialized to a Non-LWC Matrix"), (FMatrix44f)Struct1.Matrix, Struct2.Matrix);
-		Test.TestEqual<FPlane4f>(TEXT("LWC Plane must be deserialized to a Non-LWC Plane"), (FPlane4f)Struct1.Plane, Struct2.Plane);
-		Test.TestEqual<FQuat4f>(TEXT("LWC Quat must be deserialized to a Non-LWC Quat"), (FQuat4f)Struct1.Quat, Struct2.Quat);
-		Test.TestEqual<FRotator3f>(TEXT("LWC Rotator must be deserialized to a Non-LWC Rotator"), (FRotator3f)Struct1.Rotator, Struct2.Rotator);
-		Test.TestTrue(TEXT("LWC Transform must be deserialized to a Non-LWC Transform"), Struct2.Transform.Equals((FTransform3f)Struct1.Transform));
-		Test.TestEqual<FBox3f>(TEXT("LWC Box must be deserialized to a Non-LWC Box"), (FBox3f)Struct1.Box, Struct2.Box);
-		Test.TestEqual<FBox2f>(TEXT("LWC Box2D must be deserialized to a Non-LWC Box2D"), (FBox2f)Struct1.Box2D, Struct2.Box2D);
-		Test.TestEqual<FBoxSphereBounds3f>(TEXT("LWC BoxSphereBounds must be deserialized to a Non-LWC BoxSphereBounds"), (FBoxSphereBounds3f)Struct1.BoxSphereBounds, Struct2.BoxSphereBounds);
-		Test.TestEqual<float>(TEXT("LWC struct float must be the same when deserialized to a non-LWC struct float"), Struct1.Float, Struct2.Float);
-		Test.TestEqual<double>(TEXT("LWC struct double must be the same when deserialized to a non-LWC struct double"), Struct1.Double, Struct2.Double);
+		CHECK_EQUALS(TEXT("LWC Vector must be deserialized to a Non-LWC Vector"), (FVector3f)Struct1.Vector, Struct2.Vector);
+		CHECK_EQUALS(TEXT("LWC Vector2D must be deserialized to a Non-LWC Vector2D"), (FVector2f)Struct1.Vector2D, Struct2.Vector2D);
+		CHECK_EQUALS(TEXT("LWC Vector4 must be deserialized to a Non-LWC Vector4"), (FVector4f)Struct1.Vector4, Struct2.Vector4);
+		CHECK_EQUALS(TEXT("LWC Matrix must be deserialized to a Non-LWC Matrix"), (FMatrix44f)Struct1.Matrix, Struct2.Matrix);
+		CHECK_EQUALS(TEXT("LWC Plane must be deserialized to a Non-LWC Plane"), (FPlane4f)Struct1.Plane, Struct2.Plane);
+		CHECK_EQUALS(TEXT("LWC Quat must be deserialized to a Non-LWC Quat"), (FQuat4f)Struct1.Quat, Struct2.Quat);
+		CHECK_EQUALS(TEXT("LWC Rotator must be deserialized to a Non-LWC Rotator"), (FRotator3f)Struct1.Rotator, Struct2.Rotator);
+		CHECK_MESSAGE(TEXT("LWC Transform must be deserialized to a Non-LWC Transform"), Struct2.Transform.Equals((FTransform3f)Struct1.Transform));
+		CHECK_EQUALS(TEXT("LWC Box must be deserialized to a Non-LWC Box"), (FBox3f)Struct1.Box, Struct2.Box);
+		CHECK_EQUALS(TEXT("LWC Box2D must be deserialized to a Non-LWC Box2D"), (FBox2f)Struct1.Box2D, Struct2.Box2D);
+		CHECK_EQUALS(TEXT("LWC BoxSphereBounds must be deserialized to a Non-LWC BoxSphereBounds"), (FBoxSphereBounds3f)Struct1.BoxSphereBounds, Struct2.BoxSphereBounds);
+		CHECK_EQUALS(TEXT("LWC struct float must be the same when deserialized to a non-LWC struct float"), Struct1.Float, Struct2.Float);
+		CHECK_EQUALS(TEXT("LWC struct double must be the same when deserialized to a non-LWC struct double"), Struct1.Double, Struct2.Double);
 
 		const bool bIsOrientedBoxEqual = Struct1.OrientedBox.AxisX == FVector(Struct2.OrientedBox.AxisX)
 			&& Struct1.OrientedBox.AxisY == FVector(Struct2.OrientedBox.AxisY)
@@ -166,7 +170,7 @@ namespace StructSerializerTest
 			&& Struct1.OrientedBox.ExtentX == Struct2.OrientedBox.ExtentX
 			&& Struct1.OrientedBox.ExtentY == Struct2.OrientedBox.ExtentY
 			&& Struct1.OrientedBox.ExtentZ == Struct2.OrientedBox.ExtentZ;
-		Test.TestTrue(TEXT("LWC OrientedBox must be deserialized to a Non-LWC OrientedBox"), bIsOrientedBoxEqual);
+		CHECK_MESSAGE(TEXT("LWC OrientedBox must be deserialized to a Non-LWC OrientedBox"), bIsOrientedBoxEqual);
 
 		//Container testing
 		bool bAreArrayEqual = Struct1.VectorArray.Num() == Struct2.VectorArray.Num();
@@ -177,7 +181,7 @@ namespace StructSerializerTest
 				bAreArrayEqual = false;
 			}
 		}
-		Test.TestTrue(TEXT("Array of LWC Vectors must be deserialized to an Array of Non-LWC Vectors"), bAreArrayEqual);
+		CHECK_MESSAGE(TEXT("Array of LWC Vectors must be deserialized to an Array of Non-LWC Vectors"), bAreArrayEqual);
 
 		bool bAreMapsEqual = Struct1.StrToVec.Num() == Struct2.StrToVec.Num();
 		for (const TPair<FString, FVector>& Struct1Pair : Struct1.StrToVec)
@@ -196,7 +200,7 @@ namespace StructSerializerTest
 				break;
 			}
 		}
-		Test.TestTrue(TEXT("Map of LWC Vectors must be deserialized to a Map of Non-LWC Vectors"), bAreMapsEqual);
+		CHECK_MESSAGE(TEXT("Map of LWC Vectors must be deserialized to a Map of Non-LWC Vectors"), bAreMapsEqual);
 
 		bool bAreSetsEqual = Struct1.VectorSet.Num() == Struct2.VectorSet.Num();
 		for (const FVector& Vector : Struct1.VectorSet)
@@ -209,26 +213,26 @@ namespace StructSerializerTest
 				break;
 			}
 		}
-		Test.TestTrue(TEXT("Set of LWC Vectors must be deserialized to a Set of Non-LWC Vectors"), bAreSetsEqual);
+		CHECK_MESSAGE(TEXT("Set of LWC Vectors must be deserialized to a Set of Non-LWC Vectors"), bAreSetsEqual);
 	}
 
-	void ValidateLWCDeserializationBackwardCompatibility(FAutomationTestBase& Test, const FStructSerializerNonLWCTypesTest& Struct1, const FStructSerializerLWCTypesTest& Struct2)
+	void ValidateLWCDeserializationBackwardCompatibility(const FStructSerializerNonLWCTypesTest& Struct1, const FStructSerializerLWCTypesTest& Struct2)
 	{
 		//Make comparison by casting the float (non-lwc) version up to double (lwc) since this is what will happen during deserialization
-		Test.TestEqual<FVector>(TEXT("Non-LWC Vector must be deserialized to a LWC Vector"), (FVector)Struct1.Vector, Struct2.Vector);
-		Test.TestEqual<FVector2D>(TEXT("Non-LWC Vector2D must be deserialized to a LWC Vector2D"), (FVector2D)Struct1.Vector2D, Struct2.Vector2D);
-		Test.TestEqual<FVector4>(TEXT("Non-LWC Vector4 must be deserialized toa LWC Vector4"), (FVector4)Struct1.Vector4, Struct2.Vector4);
-		Test.TestEqual<FMatrix>(TEXT("Non-LWC Matrix must be deserialized to a LWC Matrix"), (FMatrix)Struct1.Matrix, Struct2.Matrix);
-		Test.TestEqual<FPlane>(TEXT("Non-LWC Plane must be deserialized to a LWC Plane"), (FPlane)Struct1.Plane, Struct2.Plane);
-		Test.TestEqual<FQuat>(TEXT("Non-LWC Quat must be deserialized to a LWC Quat"), (FQuat)Struct1.Quat, Struct2.Quat);
-		Test.TestEqual<FRotator>(TEXT("Non-LWC Rotator must be deserialized to a LWC Rotator"), (FRotator)Struct1.Rotator, Struct2.Rotator);
-		Test.TestTrue(TEXT("Non-LWC Transform must be deserialized to a LWC Transform"), Struct2.Transform.Equals((FTransform)Struct1.Transform));
-		Test.TestEqual<FBox>(TEXT("Non-LWC Box must be deserialized to a LWC Box"), (FBox)Struct1.Box, Struct2.Box);
-
-		Test.TestEqual<FBox2D>(TEXT("Non-LWC Box2D must be deserialized to a LWC Box2D"), (FBox2D)Struct1.Box2D, Struct2.Box2D);
-		Test.TestEqual<FBoxSphereBounds>(TEXT("Non-LWC BoxSphereBounds must be deserialized to a LWC BoxSphereBounds"), (FBoxSphereBounds)Struct1.BoxSphereBounds, Struct2.BoxSphereBounds);
-		Test.TestEqual<float>(TEXT("Non-LWC struct float must be the same when deserialized to a LWC struct float"), Struct1.Float, Struct2.Float);
-		Test.TestEqual<double>(TEXT("Non-LWC struct double must be the same when deserialized to a LWC struct double"), Struct1.Double, Struct2.Double);
+		CHECK_EQUALS(TEXT("Non-LWC Vector must be deserialized to a LWC Vector"), (FVector)Struct1.Vector, Struct2.Vector);
+		CHECK_EQUALS(TEXT("Non-LWC Vector2D must be deserialized to a LWC Vector2D"), (FVector2D)Struct1.Vector2D, Struct2.Vector2D);
+		CHECK_EQUALS(TEXT("Non-LWC Vector4 must be deserialized toa LWC Vector4"), (FVector4)Struct1.Vector4, Struct2.Vector4);
+		CHECK_EQUALS(TEXT("Non-LWC Matrix must be deserialized to a LWC Matrix"), (FMatrix)Struct1.Matrix, Struct2.Matrix);
+		CHECK_EQUALS(TEXT("Non-LWC Plane must be deserialized to a LWC Plane"), (FPlane)Struct1.Plane, Struct2.Plane);
+		CHECK_EQUALS(TEXT("Non-LWC Quat must be deserialized to a LWC Quat"), (FQuat)Struct1.Quat, Struct2.Quat);
+		CHECK_EQUALS(TEXT("Non-LWC Rotator must be deserialized to a LWC Rotator"), (FRotator)Struct1.Rotator, Struct2.Rotator);
+		CHECK_MESSAGE(TEXT("Non-LWC Transform must be deserialized to a LWC Transform"), Struct2.Transform.Equals((FTransform)Struct1.Transform));
+		CHECK_EQUALS(TEXT("Non-LWC Box must be deserialized to a LWC Box"), (FBox)Struct1.Box, Struct2.Box);
+
+		CHECK_EQUALS(TEXT("Non-LWC Box2D must be deserialized to a LWC Box2D"), (FBox2D)Struct1.Box2D, Struct2.Box2D);
+		CHECK_EQUALS(TEXT("Non-LWC BoxSphereBounds must be deserialized to a LWC BoxSphereBounds"), (FBoxSphereBounds)Struct1.BoxSphereBounds, Struct2.BoxSphereBounds);
+		CHECK_EQUALS(TEXT("Non-LWC struct float must be the same when deserialized to a LWC struct float"), Struct1.Float, Struct2.Float);
+		CHECK_EQUALS(TEXT("Non-LWC struct double must be the same when deserialized to a LWC struct double"), Struct1.Double, Struct2.Double);
 
 		const bool bIsOrientedBoxEqual = FVector(Struct1.OrientedBox.AxisX) == Struct2.OrientedBox.AxisX
 			&& FVector(Struct1.OrientedBox.AxisY) == Struct2.OrientedBox.AxisY
@@ -237,7 +241,7 @@ namespace StructSerializerTest
 			&& Struct1.OrientedBox.ExtentX == Struct2.OrientedBox.ExtentX
 			&& Struct1.OrientedBox.ExtentY == Struct2.OrientedBox.ExtentY
 			&& Struct1.OrientedBox.ExtentZ == Struct2.OrientedBox.ExtentZ;
-		Test.TestTrue(TEXT("Non-LWC OrientedBox must be deserialized to a LWC OrientedBox"), bIsOrientedBoxEqual);
+		CHECK_MESSAGE(TEXT("Non-LWC OrientedBox must be deserialized to a LWC OrientedBox"), bIsOrientedBoxEqual);
 
 		//Container testing
 		bool bAreArrayEqual = Struct1.VectorArray.Num() == Struct2.VectorArray.Num();
@@ -248,7 +252,7 @@ namespace StructSerializerTest
 				bAreArrayEqual = false;
 			}
 		}
-		Test.TestTrue(TEXT("Array of Non-LWC Vectors must be deserialized to an Array of LWC Vectors"), bAreArrayEqual);
+		CHECK_MESSAGE(TEXT("Array of Non-LWC Vectors must be deserialized to an Array of LWC Vectors"), bAreArrayEqual);
 
 		bool bAreMapsEqual = Struct1.StrToVec.Num() == Struct2.StrToVec.Num();
 		for (const TPair<FString, FVector3f>& Struct1Pair : Struct1.StrToVec)
@@ -267,7 +271,7 @@ namespace StructSerializerTest
 				break;
 			}
 		}
-		Test.TestTrue(TEXT("Map of Non-LWC Vectors must be deserialized to a Map of LWC Vectors"), bAreMapsEqual);
+		CHECK_MESSAGE(TEXT("Map of Non-LWC Vectors must be deserialized to a Map of LWC Vectors"), bAreMapsEqual);
 
 		bool bAreSetsEqual = Struct1.VectorSet.Num() == Struct2.VectorSet.Num();
 		for (const FVector3f& Vector : Struct1.VectorSet)
@@ -280,25 +284,25 @@ namespace StructSerializerTest
 				break;
 			}
 		}
-		Test.TestTrue(TEXT("Set of Non-LWC Vectors must be deserialized to a Set of LWC Vectors"), bAreSetsEqual);
+		CHECK_MESSAGE(TEXT("Set of Non-LWC Vectors must be deserialized to a Set of LWC Vectors"), bAreSetsEqual);
 	}
 
-	void ValidateLWCTypes(FAutomationTestBase& Test, const FStructSerializerLWCTypesTest& Struct1, const FStructSerializerLWCTypesTest& Struct2)
+	void ValidateLWCTypes(const FStructSerializerLWCTypesTest& Struct1, const FStructSerializerLWCTypesTest& Struct2)
 	{
-		Test.TestEqual<FVector>(TEXT("LWC Vector must be the same before and after de-serialization"), Struct1.Vector, Struct2.Vector);
-		Test.TestEqual<FVector2D>(TEXT("LWC Vector2D must be the same before and after de-serialization"), Struct1.Vector2D, Struct2.Vector2D);
-		Test.TestEqual<FVector4>(TEXT("LWC Vector4 must be the same before and after de-serialization"), Struct1.Vector4, Struct2.Vector4);
-		Test.TestEqual<FMatrix>(TEXT("LWC Matrix must be the same before and after de-serialization"), Struct1.Matrix, Struct2.Matrix);
-		Test.TestEqual<FPlane>(TEXT("LWC Plane must be the same before and after de-serialization"), Struct1.Plane, Struct2.Plane);
-		Test.TestEqual<FQuat>(TEXT("LWC Quat must be the same before and after de-serialization"), Struct1.Quat, Struct2.Quat);
-		Test.TestEqual<FRotator>(TEXT("LWC Rotator must be the same before and after de-serialization"), Struct1.Rotator, Struct2.Rotator);
-		Test.TestTrue(TEXT("LWC Transform must be the same before and after de-serialization"), Struct2.Transform.Equals(Struct1.Transform));
-		Test.TestEqual<FBox>(TEXT("LWC Box must be the same before and after de-serialization"), Struct1.Box, Struct2.Box);
-
-		Test.TestEqual<FBox2D>(TEXT("LWC Box2D must be the same before and after de-serialization"), Struct1.Box2D, Struct2.Box2D);
-		Test.TestEqual<FBoxSphereBounds>(TEXT("LWC BoxSphereBounds must be the same before and after de-serialization"), Struct1.BoxSphereBounds, Struct2.BoxSphereBounds);
-		Test.TestEqual<float>(TEXT("LWC struct float must be the same before and after de-serialization"), Struct1.Float, Struct2.Float);
-		Test.TestEqual<double>(TEXT("LWC struct double must be the same before and after de-serialization"), Struct1.Double, Struct2.Double);
+		CHECK_EQUALS(TEXT("LWC Vector must be the same before and after de-serialization"), Struct1.Vector, Struct2.Vector);
+		CHECK_EQUALS(TEXT("LWC Vector2D must be the same before and after de-serialization"), Struct1.Vector2D, Struct2.Vector2D);
+		CHECK_EQUALS(TEXT("LWC Vector4 must be the same before and after de-serialization"), Struct1.Vector4, Struct2.Vector4);
+		CHECK_EQUALS(TEXT("LWC Matrix must be the same before and after de-serialization"), Struct1.Matrix, Struct2.Matrix);
+		CHECK_EQUALS(TEXT("LWC Plane must be the same before and after de-serialization"), Struct1.Plane, Struct2.Plane);
+		CHECK_EQUALS(TEXT("LWC Quat must be the same before and after de-serialization"), Struct1.Quat, Struct2.Quat);
+		CHECK_EQUALS(TEXT("LWC Rotator must be the same before and after de-serialization"), Struct1.Rotator, Struct2.Rotator);
+		CHECK_MESSAGE(TEXT("LWC Transform must be the same before and after de-serialization"), Struct2.Transform.Equals(Struct1.Transform));
+		CHECK_EQUALS(TEXT("LWC Box must be the same before and after de-serialization"), Struct1.Box, Struct2.Box);
+
+		CHECK_EQUALS(TEXT("LWC Box2D must be the same before and after de-serialization"), Struct1.Box2D, Struct2.Box2D);
+		CHECK_EQUALS(TEXT("LWC BoxSphereBounds must be the same before and after de-serialization"), Struct1.BoxSphereBounds, Struct2.BoxSphereBounds);
+		CHECK_EQUALS(TEXT("LWC struct float must be the same before and after de-serialization"), Struct1.Float, Struct2.Float);
+		CHECK_EQUALS(TEXT("LWC struct double must be the same before and after de-serialization"), Struct1.Double, Struct2.Double);
 
 		const bool bIsOrientedBoxEqual = Struct1.OrientedBox.AxisX == Struct2.OrientedBox.AxisX
 			&& Struct1.OrientedBox.AxisY == Struct2.OrientedBox.AxisY
@@ -307,20 +311,20 @@ namespace StructSerializerTest
 			&& Struct1.OrientedBox.ExtentX == Struct2.OrientedBox.ExtentX
 			&& Struct1.OrientedBox.ExtentY == Struct2.OrientedBox.ExtentY
 			&& Struct1.OrientedBox.ExtentZ == Struct2.OrientedBox.ExtentZ;
-		Test.TestTrue(TEXT("LWC OrientedBox must be the same before and after de-serialization"), bIsOrientedBoxEqual);
+		CHECK_MESSAGE(TEXT("LWC OrientedBox must be the same before and after de-serialization"), bIsOrientedBoxEqual);
 
-		Test.TestEqual<TArray<FVector>>(TEXT("LWC test - Arrays.VectorArray must be the same before and after de-/serialization"), Struct1.VectorArray, Struct2.VectorArray);
-		Test.TestTrue(TEXT("LWC test - Maps.StrToVec must be the same before and after de-/serialization"), Struct1.StrToVec.OrderIndependentCompareEqual(Struct2.StrToVec));
-		Test.TestTrue(TEXT("LWC test - Sets.VectorSet must be the same before and after de-/serialization"), Struct1.VectorSet.Num() == Struct2.VectorSet.Num() && Struct1.VectorSet.Difference(Struct2.VectorSet).Num() == 0);
+		CHECK_EQUALS(TEXT("LWC test - Arrays.VectorArray must be the same before and after de-/serialization"), Struct1.VectorArray, Struct2.VectorArray);
+		CHECK_MESSAGE(TEXT("LWC test - Maps.StrToVec must be the same before and after de-/serialization"), Struct1.StrToVec.OrderIndependentCompareEqual(Struct2.StrToVec));
+		CHECK_MESSAGE(TEXT("LWC test - Sets.VectorSet must be the same before and after de-/serialization"), Struct1.VectorSet.Num() == Struct2.VectorSet.Num() && Struct1.VectorSet.Difference(Struct2.VectorSet).Num() == 0);
 	}
 
 	template<typename TSerializerBackend, typename TDeserializerBackend>
-	void TestElementSerialization(FAutomationTestBase& Test)
+	void TestElementSerialization()
 	{
 		// serialization
 	
 		FStructSerializerTestStruct OriginalStruct;
-		UClass* MetaDataClass = LoadClass<UMetaData>(nullptr, TEXT("/Script/CoreUObject.MetaData"));
+		UClass* MetaDataClass = UMetaData::StaticClass();
 		UMetaData* MetaDataObject = NewObject<UMetaData>();
 		// setup object tests
 		OriginalStruct.Objects.RawClass = MetaDataClass;
@@ -358,9 +362,9 @@ namespace StructSerializerTest
 				FStructSerializer::SerializeElement(&TestStruct, Property, INDEX_NONE, SerializerBackend, Policies);
 
 				FStructSerializerTestStruct TestStruct2(NoInit);
-				Test.TestTrue(TEXT("Deserialization must succeed"), FStructDeserializer::DeserializeElement(&TestStruct2, *FStructSerializerTestStruct::StaticStruct(), INDEX_NONE, DeserializerBackend, DeserializerPolicies));
+				CHECK_MESSAGE(TEXT("Deserialization must succeed"), FStructDeserializer::DeserializeElement(&TestStruct2, *FStructSerializerTestStruct::StaticStruct(), INDEX_NONE, DeserializerBackend, DeserializerPolicies));
 
-				ValidateNumerics(Test, TestStruct.Numerics, TestStruct2.Numerics);
+				ValidateNumerics(TestStruct.Numerics, TestStruct2.Numerics);
 			}
 
 			//Booleans
@@ -378,9 +382,9 @@ namespace StructSerializerTest
 				FStructSerializer::SerializeElement(&TestStruct, Property, INDEX_NONE, SerializerBackend, Policies);
 
 				FStructSerializerTestStruct TestStruct2(NoInit);
-				Test.TestTrue(TEXT("Deserialization must succeed"), FStructDeserializer::DeserializeElement(&TestStruct2, *FStructSerializerTestStruct::StaticStruct(), INDEX_NONE, DeserializerBackend, DeserializerPolicies));
+				CHECK_MESSAGE(TEXT("Deserialization must succeed"), FStructDeserializer::DeserializeElement(&TestStruct2, *FStructSerializerTestStruct::StaticStruct(), INDEX_NONE, DeserializerBackend, DeserializerPolicies));
 
-				ValidateBooleans(Test, TestStruct.Booleans, TestStruct2.Booleans);
+				ValidateBooleans(TestStruct.Booleans, TestStruct2.Booleans);
 			}
 
 			//Objects
@@ -398,9 +402,9 @@ namespace StructSerializerTest
 				FStructSerializer::SerializeElement(&TestStruct, Property, INDEX_NONE, SerializerBackend, Policies);
 
 				FStructSerializerTestStruct TestStruct2(NoInit);
-				Test.TestTrue(TEXT("Deserialization must succeed"), FStructDeserializer::DeserializeElement(&TestStruct2, *FStructSerializerTestStruct::StaticStruct(), INDEX_NONE, DeserializerBackend, DeserializerPolicies));
+				CHECK_MESSAGE(TEXT("Deserialization must succeed"), FStructDeserializer::DeserializeElement(&TestStruct2, *FStructSerializerTestStruct::StaticStruct(), INDEX_NONE, DeserializerBackend, DeserializerPolicies));
 
-				ValidateObjects(Test, TestStruct.Objects, TestStruct2.Objects);
+				ValidateObjects(TestStruct.Objects, TestStruct2.Objects);
 			}
 
 			//Built ins
@@ -418,9 +422,9 @@ namespace StructSerializerTest
 				FStructSerializer::SerializeElement(&TestStruct, Property, INDEX_NONE, SerializerBackend, Policies);
 
 				FStructSerializerTestStruct TestStruct2(NoInit);
-				Test.TestTrue(TEXT("Deserialization must succeed"), FStructDeserializer::DeserializeElement(&TestStruct2, *FStructSerializerTestStruct::StaticStruct(), INDEX_NONE, DeserializerBackend, DeserializerPolicies));
+				CHECK_MESSAGE(TEXT("Deserialization must succeed"), FStructDeserializer::DeserializeElement(&TestStruct2, *FStructSerializerTestStruct::StaticStruct(), INDEX_NONE, DeserializerBackend, DeserializerPolicies));
 
-				ValidateBuiltIns(Test, TestStruct.Builtins, TestStruct2.Builtins);
+				ValidateBuiltIns(TestStruct.Builtins, TestStruct2.Builtins);
 			}
 
 			//Arrays
@@ -443,9 +447,9 @@ namespace StructSerializerTest
 				TestStruct2.Arrays.VectorArray.SetNumUninitialized(TestStruct.Arrays.VectorArray.Num());
 				TestStruct2.Arrays.StructArray.SetNumZeroed(TestStruct.Arrays.StructArray.Num()); 
 
-				Test.TestTrue(TEXT("Deserialization must succeed"), FStructDeserializer::DeserializeElement(&TestStruct2, *FStructSerializerTestStruct::StaticStruct(), INDEX_NONE, DeserializerBackend, DeserializerPolicies));
+				CHECK_MESSAGE(TEXT("Deserialization must succeed"), FStructDeserializer::DeserializeElement(&TestStruct2, *FStructSerializerTestStruct::StaticStruct(), INDEX_NONE, DeserializerBackend, DeserializerPolicies));
 
-				ValidateArrays(Test, TestStruct.Arrays, TestStruct2.Arrays);
+				ValidateArrays(TestStruct.Arrays, TestStruct2.Arrays);
 			}
 
 			//Maps
@@ -468,9 +472,9 @@ namespace StructSerializerTest
 				CopyKeys(TestStruct2.Maps.StrToStruct, TestStruct.Maps.StrToStruct);
 				CopyKeys(TestStruct2.Maps.StrToVec, TestStruct.Maps.StrToVec);
 
-				Test.TestTrue(TEXT("Deserialization must succeed"), FStructDeserializer::DeserializeElement(&TestStruct2, *FStructSerializerTestStruct::StaticStruct(), INDEX_NONE, DeserializerBackend, DeserializerPolicies));
+				CHECK_MESSAGE(TEXT("Deserialization must succeed"), FStructDeserializer::DeserializeElement(&TestStruct2, *FStructSerializerTestStruct::StaticStruct(), INDEX_NONE, DeserializerBackend, DeserializerPolicies));
 
-				ValidateMaps(Test, TestStruct.Maps, TestStruct2.Maps);
+				ValidateMaps(TestStruct.Maps, TestStruct2.Maps);
 			}
 
 			//Sets
@@ -494,9 +498,9 @@ namespace StructSerializerTest
 				TestStruct2.Sets.NameSet = { TEXT("Pre1"), TEXT("Pre2") , TEXT("Pre3") };
 				TestStruct2.Sets.StrSet = { TEXT("Pre4"), TEXT("Pre5") , TEXT("Pre6"), TEXT("Pre7") };
 				TestStruct2.Sets.StructSet = { FStructSerializerBuiltinTestStruct(NoInit) };
-				Test.TestTrue(TEXT("Deserialization must succeed"), FStructDeserializer::DeserializeElement(&TestStruct2, *FStructSerializerTestStruct::StaticStruct(), INDEX_NONE, DeserializerBackend, DeserializerPolicies));
+				CHECK_MESSAGE(TEXT("Deserialization must succeed"), FStructDeserializer::DeserializeElement(&TestStruct2, *FStructSerializerTestStruct::StaticStruct(), INDEX_NONE, DeserializerBackend, DeserializerPolicies));
 
-				ValidateSets(Test, TestStruct.Sets, TestStruct2.Sets);
+				ValidateSets(TestStruct.Sets, TestStruct2.Sets);
 			}
 
 			//TArray<uint8> element
@@ -520,10 +524,10 @@ namespace StructSerializerTest
 
 				constexpr int32 TestIndex = 1;
 				FStructSerializer::SerializeElement(&TestStruct, Property, TestIndex, SerializerBackend, Policies);
-				Test.TestTrue(TEXT("Deserialization must succeed"), FStructDeserializer::DeserializeElement(&TestStruct2, *FStructSerializerArrayTestStruct::StaticStruct(), TestIndex, DeserializerBackend, DeserializerPolicies));
-				Test.TestNotEqual<uint8>(TEXT("Arrays.ByteArray[0] must not be the same before and after de-/serialization of element 1"), TestStruct.ByteArray[0], TestStruct2.ByteArray[0]);
-				Test.TestEqual<uint8>(TEXT("Arrays.ByteArray[1] must be the same before and after de-/serialization of element 1"), TestStruct.ByteArray[TestIndex], TestStruct2.ByteArray[TestIndex]);
-				Test.TestNotEqual<uint8>(TEXT("Arrays.ByteArray[2] must not be the same before and after de-/serialization of element 1"), TestStruct.ByteArray[2], TestStruct2.ByteArray[2]);
+				CHECK_MESSAGE(TEXT("Deserialization must succeed"), FStructDeserializer::DeserializeElement(&TestStruct2, *FStructSerializerArrayTestStruct::StaticStruct(), TestIndex, DeserializerBackend, DeserializerPolicies));
+				CHECK_NOT_EQUALS(TEXT("Arrays.ByteArray[0] must not be the same before and after de-/serialization of element 1"), TestStruct.ByteArray[0], TestStruct2.ByteArray[0]);
+				CHECK_EQUALS(TEXT("Arrays.ByteArray[1] must be the same before and after de-/serialization of element 1"), TestStruct.ByteArray[TestIndex], TestStruct2.ByteArray[TestIndex]);
+				CHECK_NOT_EQUALS(TEXT("Arrays.ByteArray[2] must not be the same before and after de-/serialization of element 1"), TestStruct.ByteArray[2], TestStruct2.ByteArray[2]);
 			}
 
 			//TArray<Struct> element
@@ -544,9 +548,9 @@ namespace StructSerializerTest
 
 				constexpr int32 TestIndex = 1;
 				FStructSerializer::SerializeElement(&TestStruct, Property, TestIndex, SerializerBackend, Policies);
-				Test.TestTrue(TEXT("Deserialization must succeed"), FStructDeserializer::DeserializeElement(&TestStruct2, *FStructSerializerArrayTestStruct::StaticStruct(), TestIndex, DeserializerBackend, DeserializerPolicies));
-				Test.TestFalse(TEXT("Arrays.StructArray[0] must not be the same before and after de-/serialization of element 1"), TestStruct.StructArray[0] == TestStruct2.StructArray[0]);
-				ValidateBuiltIns(Test, TestStruct.StructArray[TestIndex], TestStruct2.StructArray[TestIndex]);
+				CHECK_MESSAGE(TEXT("Deserialization must succeed"), FStructDeserializer::DeserializeElement(&TestStruct2, *FStructSerializerArrayTestStruct::StaticStruct(), TestIndex, DeserializerBackend, DeserializerPolicies));
+				CHECK_FALSE_MESSAGE(TEXT("Arrays.StructArray[0] must not be the same before and after de-/serialization of element 1"), TestStruct.StructArray[0] == TestStruct2.StructArray[0]);
+				ValidateBuiltIns(TestStruct.StructArray[TestIndex], TestStruct2.StructArray[TestIndex]);
 			}
 			
 			//static single element
@@ -567,8 +571,8 @@ namespace StructSerializerTest
 
 				constexpr int32 TestIndex = 0;
 				FStructSerializer::SerializeElement(&TestStruct, Property, TestIndex, SerializerBackend, Policies);
-				Test.TestTrue(TEXT("Deserialization must succeed"), FStructDeserializer::DeserializeElement(&TestStruct2, *FStructSerializerArrayTestStruct::StaticStruct(), TestIndex, DeserializerBackend, DeserializerPolicies));
-				Test.TestEqual<int32>(TEXT("Arrays.StaticSingleElement[0] must be the same before and after de-/serialization"), TestStruct.StaticSingleElement[TestIndex], TestStruct2.StaticSingleElement[TestIndex]);
+				CHECK_MESSAGE(TEXT("Deserialization must succeed"), FStructDeserializer::DeserializeElement(&TestStruct2, *FStructSerializerArrayTestStruct::StaticStruct(), TestIndex, DeserializerBackend, DeserializerPolicies));
+				CHECK_EQUALS(TEXT("Arrays.StaticSingleElement[0] must be the same before and after de-/serialization"), TestStruct.StaticSingleElement[TestIndex], TestStruct2.StaticSingleElement[TestIndex]);
 			}
 
 			//static float array element
@@ -589,10 +593,10 @@ namespace StructSerializerTest
 
 				constexpr int32 TestIndex = 1;
 				FStructSerializer::SerializeElement(&TestStruct, Property, TestIndex, SerializerBackend, Policies);
-				Test.TestTrue(TEXT("Deserialization must succeed"), FStructDeserializer::DeserializeElement(&TestStruct2, *FStructSerializerArrayTestStruct::StaticStruct(), TestIndex, DeserializerBackend, DeserializerPolicies));
-				Test.TestNotEqual<float>(TEXT("Arrays.StaticFloatArray[0] must not be the same before and after de-/serialization of element 1"), TestStruct.StaticFloatArray[0], TestStruct2.StaticFloatArray[0]);
-				Test.TestEqual<float>(TEXT("Arrays.StaticFloatArray[1] must be the same before and after de-/serialization"), TestStruct.StaticFloatArray[TestIndex], TestStruct2.StaticFloatArray[TestIndex]);
-				Test.TestNotEqual<float>(TEXT("Arrays.StaticFloatArray[2] must not be the same before and after de-/serialization of element 1"), TestStruct.StaticFloatArray[2], TestStruct2.StaticFloatArray[2]);
+				CHECK_MESSAGE(TEXT("Deserialization must succeed"), FStructDeserializer::DeserializeElement(&TestStruct2, *FStructSerializerArrayTestStruct::StaticStruct(), TestIndex, DeserializerBackend, DeserializerPolicies));
+				CHECK_NOT_EQUALS(TEXT("Arrays.StaticFloatArray[0] must not be the same before and after de-/serialization of element 1"), TestStruct.StaticFloatArray[0], TestStruct2.StaticFloatArray[0]);
+				CHECK_EQUALS(TEXT("Arrays.StaticFloatArray[1] must be the same before and after de-/serialization"), TestStruct.StaticFloatArray[TestIndex], TestStruct2.StaticFloatArray[TestIndex]);
+				CHECK_NOT_EQUALS(TEXT("Arrays.StaticFloatArray[2] must not be the same before and after de-/serialization of element 1"), TestStruct.StaticFloatArray[2], TestStruct2.StaticFloatArray[2]);
 			}
 
 			//TMap<int32, FString> element
@@ -613,14 +617,14 @@ namespace StructSerializerTest
 
 				constexpr int32 TestIndex = 1;
 				FStructSerializer::SerializeElement(&TestStruct, Property, TestIndex, SerializerBackend, Policies);
-				Test.TestTrue(TEXT("Deserialization must succeed"), FStructDeserializer::DeserializeElement(&TestStruct2, *FStructSerializerMapTestStruct::StaticStruct(), TestIndex, DeserializerBackend, DeserializerPolicies));
+				CHECK_MESSAGE(TEXT("Deserialization must succeed"), FStructDeserializer::DeserializeElement(&TestStruct2, *FStructSerializerMapTestStruct::StaticStruct(), TestIndex, DeserializerBackend, DeserializerPolicies));
 
 				TArray<int32> Keys;
 				Keys.Reserve(TestStruct.IntToStr.Num());
 				TestStruct.IntToStr.GenerateKeyArray(Keys);
-				Test.TestNotEqual<FString>(TEXT("Maps.IntToStr[0] must not be the same before and after de-/serialization of element 1"), TestStruct.IntToStr[Keys[0]], TestStruct2.IntToStr[Keys[0]]);
-				Test.TestEqual<FString>(TEXT("Maps.IntToStr[1] must be the same before and after de-/serialization of element 1"), TestStruct.IntToStr[Keys[TestIndex]], TestStruct2.IntToStr[Keys[TestIndex]]);
-				Test.TestNotEqual<FString>(TEXT("Maps.IntToStr[2] must not be the same before and after de-/serialization of element 1"), TestStruct.IntToStr[Keys[2]], TestStruct2.IntToStr[Keys[2]]);
+				CHECK_NOT_EQUALS(TEXT("Maps.IntToStr[0] must not be the same before and after de-/serialization of element 1"), TestStruct.IntToStr[Keys[0]], TestStruct2.IntToStr[Keys[0]]);
+				CHECK_EQUALS(TEXT("Maps.IntToStr[1] must be the same before and after de-/serialization of element 1"), TestStruct.IntToStr[Keys[TestIndex]], TestStruct2.IntToStr[Keys[TestIndex]]);
+				CHECK_NOT_EQUALS(TEXT("Maps.IntToStr[2] must not be the same before and after de-/serialization of element 1"), TestStruct.IntToStr[Keys[2]], TestStruct2.IntToStr[Keys[2]]);
 			}
 
 			//TMap<FString, FVector> element
@@ -641,14 +645,14 @@ namespace StructSerializerTest
 
 				constexpr int32 TestIndex = 1;
 				FStructSerializer::SerializeElement(&TestStruct, Property, TestIndex, SerializerBackend, Policies);
-				Test.TestTrue(TEXT("Deserialization must succeed"), FStructDeserializer::DeserializeElement(&TestStruct2, *FStructSerializerMapTestStruct::StaticStruct(), TestIndex, DeserializerBackend, DeserializerPolicies));
+				CHECK_MESSAGE(TEXT("Deserialization must succeed"), FStructDeserializer::DeserializeElement(&TestStruct2, *FStructSerializerMapTestStruct::StaticStruct(), TestIndex, DeserializerBackend, DeserializerPolicies));
 
 				TArray<FString> Keys;
 				Keys.Reserve(TestStruct.IntToStr.Num());
 				TestStruct.StrToVec.GenerateKeyArray(Keys);
-				Test.TestNotEqual<FVector>(TEXT("Maps.StrToVec[0] must not be the same before and after de-/serialization of element 1"), TestStruct.StrToVec[Keys[0]], TestStruct2.StrToVec[Keys[0]]);
-				Test.TestEqual<FVector>(TEXT("Maps.StrToVec[1] must be the same before and after de-/serialization of element 1"), TestStruct.StrToVec[Keys[TestIndex]], TestStruct2.StrToVec[Keys[TestIndex]]);
-				Test.TestNotEqual<FVector>(TEXT("Maps.StrToVec[2] must not be the same before and after de-/serialization of element 1"), TestStruct.StrToVec[Keys[2]], TestStruct2.StrToVec[Keys[2]]);
+				CHECK_NOT_EQUALS(TEXT("Maps.StrToVec[0] must not be the same before and after de-/serialization of element 1"), TestStruct.StrToVec[Keys[0]], TestStruct2.StrToVec[Keys[0]]);
+				CHECK_EQUALS(TEXT("Maps.StrToVec[1] must be the same before and after de-/serialization of element 1"), TestStruct.StrToVec[Keys[TestIndex]], TestStruct2.StrToVec[Keys[TestIndex]]);
+				CHECK_NOT_EQUALS(TEXT("Maps.StrToVec[2] must not be the same before and after de-/serialization of element 1"), TestStruct.StrToVec[Keys[2]], TestStruct2.StrToVec[Keys[2]]);
 			}
 
 			//TSet<FName> element
@@ -669,23 +673,23 @@ namespace StructSerializerTest
 
 				constexpr int32 TestIndex = 1;
 				FStructSerializer::SerializeElement(&TestStruct, Property, TestIndex, SerializerBackend, Policies);
-				Test.TestTrue(TEXT("Deserialization must succeed"), FStructDeserializer::DeserializeElement(&TestStruct2, *FStructSerializerSetTestStruct::StaticStruct(), TestIndex, DeserializerBackend, DeserializerPolicies));
+				CHECK_MESSAGE(TEXT("Deserialization must succeed"), FStructDeserializer::DeserializeElement(&TestStruct2, *FStructSerializerSetTestStruct::StaticStruct(), TestIndex, DeserializerBackend, DeserializerPolicies));
 
 				TArray<FName> SetArray1 = TestStruct.NameSet.Array();
 				TArray<FName> SetArray2 = TestStruct2.NameSet.Array();
-				Test.TestNotEqual<FName>(TEXT("Sets.NameSet[0] must not be the same before and after de-/serialization of element 1"), SetArray1[0], SetArray2[0]);
-				Test.TestEqual<FName>(TEXT("Sets.NameSet[1] must be the same before and after de-/serialization of element 1"), SetArray1[TestIndex], SetArray2[TestIndex]);
-				Test.TestNotEqual<FName>(TEXT("Sets.NameSet[2] must not be the same before and after de-/serialization of element 1"), SetArray1[2], SetArray2[2]);
+				CHECK_NOT_EQUALS(TEXT("Sets.NameSet[0] must not be the same before and after de-/serialization of element 1"), SetArray1[0], SetArray2[0]);
+				CHECK_EQUALS(TEXT("Sets.NameSet[1] must be the same before and after de-/serialization of element 1"), SetArray1[TestIndex], SetArray2[TestIndex]);
+				CHECK_NOT_EQUALS(TEXT("Sets.NameSet[2] must not be the same before and after de-/serialization of element 1"), SetArray1[2], SetArray2[2]);
 			}
 		}
 	}
 
-	void TestSerialization( FAutomationTestBase& Test, IStructSerializerBackend& SerializerBackend, IStructDeserializerBackend& DeserializerBackend )
+	void TestSerialization(IStructSerializerBackend& SerializerBackend, IStructDeserializerBackend& DeserializerBackend )
 	{
 		// serialization
 		FStructSerializerTestStruct TestStruct;
 
-		UClass* MetaDataClass = LoadClass<UMetaData>(nullptr, TEXT("/Script/CoreUObject.MetaData"));
+		UClass* MetaDataClass = UMetaData::StaticClass();
 		UMetaData* MetaDataObject = NewObject<UMetaData>();
 		// setup object tests
 		TestStruct.Objects.RawClass = MetaDataClass;
@@ -709,34 +713,35 @@ namespace StructSerializerTest
 			FStructDeserializerPolicies Policies;
 			Policies.MissingFields = EStructDeserializerErrorPolicies::Warning;
 			
-			Test.TestTrue(TEXT("Deserialization must succeed"), FStructDeserializer::Deserialize(TestStruct2, DeserializerBackend, Policies));
+			CHECK_MESSAGE(TEXT("Deserialization must succeed"), FStructDeserializer::Deserialize(TestStruct2, DeserializerBackend, Policies));
 		}
 
 		// test numerics
-		ValidateNumerics(Test, TestStruct.Numerics, TestStruct2.Numerics);
+		ValidateNumerics(TestStruct.Numerics, TestStruct2.Numerics);
 
 		// test booleans
-		ValidateBooleans(Test, TestStruct.Booleans, TestStruct2.Booleans);
+		ValidateBooleans(TestStruct.Booleans, TestStruct2.Booleans);
 	
 		// test objects
-		ValidateObjects(Test, TestStruct.Objects, TestStruct2.Objects);
+		ValidateObjects(TestStruct.Objects, TestStruct2.Objects);
 
 		// test built-ins
-		ValidateBuiltIns(Test, TestStruct.Builtins, TestStruct2.Builtins);
+		ValidateBuiltIns(TestStruct.Builtins, TestStruct2.Builtins);
 
 		// test arrays
-		ValidateArrays(Test, TestStruct.Arrays, TestStruct2.Arrays);
+		ValidateArrays(TestStruct.Arrays, TestStruct2.Arrays);
 
 		// test maps
-		ValidateMaps(Test, TestStruct.Maps, TestStruct2.Maps);
+		ValidateMaps(TestStruct.Maps, TestStruct2.Maps);
+
 		// test sets
-		ValidateSets(Test, TestStruct.Sets, TestStruct2.Sets);
+		ValidateSets(TestStruct.Sets, TestStruct2.Sets);
 
 		//Test LWC types with standard de-serialization
-		ValidateLWCTypes(Test, TestStruct.LWCTypes, TestStruct2.LWCTypes);
+		ValidateLWCTypes(TestStruct.LWCTypes, TestStruct2.LWCTypes);
 	}
 
-	void TestLWCSerialization(FAutomationTestBase& Test, IStructSerializerBackend& SerializerBackend, IStructDeserializerBackend& DeserializerBackend)
+	void TestLWCSerialization(IStructSerializerBackend& SerializerBackend, IStructDeserializerBackend& DeserializerBackend)
 	{
 		// Serialization of LWC struct into non-LWC mode to mimick sending to older UE
 		FStructSerializerLWCTypesTest TestLWCStruct;
@@ -750,13 +755,13 @@ namespace StructSerializerTest
 			FStructDeserializerPolicies Policies;
 			Policies.MissingFields = EStructDeserializerErrorPolicies::Warning;
 
-			Test.TestTrue(TEXT("Deserialization must succeed"), FStructDeserializer::Deserialize(TestNonLWCStruct, DeserializerBackend, Policies));
+			CHECK_MESSAGE(TEXT("Deserialization must succeed"), FStructDeserializer::Deserialize(TestNonLWCStruct, DeserializerBackend, Policies));
 		}
 
-		ValidateLWCSerializationBackwardCompatibility(Test, TestLWCStruct, TestNonLWCStruct);
+		ValidateLWCSerializationBackwardCompatibility(TestLWCStruct, TestNonLWCStruct);
 	}
 
-	void TestLWCDeserialization(FAutomationTestBase& Test, IStructSerializerBackend& SerializerBackend, IStructDeserializerBackend& DeserializerBackend)
+	void TestLWCDeserialization(IStructSerializerBackend& SerializerBackend, IStructDeserializerBackend& DeserializerBackend)
 	{
 		// Serialization of non lwc struct to mimick a struct coming from older UE
 		FStructSerializerNonLWCTypesTest TestNonLWCStruct;
@@ -770,10 +775,10 @@ namespace StructSerializerTest
 			FStructDeserializerPolicies Policies;
 			Policies.MissingFields = EStructDeserializerErrorPolicies::Warning;
 
-			Test.TestTrue(TEXT("Deserialization must succeed"), FStructDeserializer::Deserialize(TestLWCStruct, DeserializerBackend, Policies));
+			CHECK_MESSAGE(TEXT("Deserialization must succeed"), FStructDeserializer::Deserialize(TestLWCStruct, DeserializerBackend, Policies));
 		}
 
-		ValidateLWCDeserializationBackwardCompatibility(Test, TestNonLWCStruct, TestLWCStruct);
+		ValidateLWCDeserializationBackwardCompatibility(TestNonLWCStruct, TestLWCStruct);
 	}
 }
 
@@ -781,12 +786,9 @@ namespace StructSerializerTest
 /* Tests
  *****************************************************************************/
 
-IMPLEMENT_SIMPLE_AUTOMATION_TEST(FStructSerializerTest, "System.Core.Serialization.StructSerializer", EAutomationTestFlags::ApplicationContextMask | EAutomationTestFlags::EngineFilter)
-
-
-bool FStructSerializerTest::RunTest( const FString& Parameters )
+TEST_CASE_NAMED(FStructSerializerTest, "System::Core::Serialization::StructSerializer", "[ApplicationContextMask][EngineFilter][StructSerializer]")
 {
-	const EStructSerializerBackendFlags TestFlags = EStructSerializerBackendFlags::Default;
+	const EStructSerializerBackendFlags BackendFlags = EStructSerializerBackendFlags::Default;
 
 	// json
 	{
@@ -794,10 +796,10 @@ bool FStructSerializerTest::RunTest( const FString& Parameters )
 		FMemoryReader Reader(Buffer);
 		FMemoryWriter Writer(Buffer);
 
-		FJsonStructSerializerBackend SerializerBackend(Writer, TestFlags);
+		FJsonStructSerializerBackend SerializerBackend(Writer, BackendFlags);
 		FJsonStructDeserializerBackend DeserializerBackend(Reader);
 
-		StructSerializerTest::TestSerialization(*this, SerializerBackend, DeserializerBackend);
+		StructSerializerTest::TestSerialization(SerializerBackend, DeserializerBackend);
 
 		// uncomment this to look at the serialized data
 		//GLog->Logf(TEXT("%s"), (TCHAR*)Buffer.GetData());
@@ -808,10 +810,10 @@ bool FStructSerializerTest::RunTest( const FString& Parameters )
 		FMemoryReader Reader(Buffer);
 		FMemoryWriter Writer(Buffer);
 
-		FCborStructSerializerBackend SerializerBackend(Writer, TestFlags);
+		FCborStructSerializerBackend SerializerBackend(Writer, BackendFlags);
 		FCborStructDeserializerBackend DeserializerBackend(Reader);
 
-		StructSerializerTest::TestSerialization(*this, SerializerBackend, DeserializerBackend);
+		StructSerializerTest::TestSerialization(SerializerBackend, DeserializerBackend);
 	}
 	// cbor standard compliant endianness (big endian)
 	{
@@ -822,7 +824,7 @@ bool FStructSerializerTest::RunTest( const FString& Parameters )
 		FCborStructSerializerBackend SerializerBackend(Writer, EStructSerializerBackendFlags::Default | EStructSerializerBackendFlags::WriteCborStandardEndianness);
 		FCborStructDeserializerBackend DeserializerBackend(Reader, ECborEndianness::StandardCompliant);
 
-		StructSerializerTest::TestSerialization(*this, SerializerBackend, DeserializerBackend);
+		StructSerializerTest::TestSerialization(SerializerBackend, DeserializerBackend);
 	}
 	// cbor LWC (UE5) to NonLWC (UE4)
 	{
@@ -835,7 +837,7 @@ bool FStructSerializerTest::RunTest( const FString& Parameters )
 		constexpr bool bIsLWCCompatibilityMode = false;
 		FCborStructDeserializerBackend DeserializerBackend(Reader, ECborEndianness::StandardCompliant, bIsLWCCompatibilityMode);
 
-		StructSerializerTest::TestLWCSerialization(*this, SerializerBackend, DeserializerBackend);
+		StructSerializerTest::TestLWCSerialization(SerializerBackend, DeserializerBackend);
 	}
 	// cbor Non LWC (UE4) to LWC (UE5)
 	{
@@ -848,28 +850,20 @@ bool FStructSerializerTest::RunTest( const FString& Parameters )
 		constexpr bool bIsLWCCompatibilityMode = true;
 		FCborStructDeserializerBackend DeserializerBackend(Reader, ECborEndianness::StandardCompliant, bIsLWCCompatibilityMode);
 
-		StructSerializerTest::TestLWCDeserialization(*this, SerializerBackend, DeserializerBackend);
+		StructSerializerTest::TestLWCDeserialization(SerializerBackend, DeserializerBackend);
 	}
-
-	return true;
 }
 
-IMPLEMENT_SIMPLE_AUTOMATION_TEST(FStructElementSerializerTest, "System.Core.Serialization.StructElementSerializer", EAutomationTestFlags::ApplicationContextMask | EAutomationTestFlags::EngineFilter)
-
-bool FStructElementSerializerTest::RunTest(const FString& Parameters)
+TEST_CASE_NAMED(FStructElementSerializerTest, "System::Core::Serialization::StructElementSerializer", "[ApplicationContextMask][EngineFilter][StructSerializer]")
 {
 	//Element de/serialization for both types of backend
 	{
-		StructSerializerTest::TestElementSerialization<FJsonStructSerializerBackend, FJsonStructDeserializerBackend>(*this);
-		StructSerializerTest::TestElementSerialization<FCborStructSerializerBackend, FCborStructDeserializerBackend>(*this);
+		StructSerializerTest::TestElementSerialization<FJsonStructSerializerBackend, FJsonStructDeserializerBackend>();
+		StructSerializerTest::TestElementSerialization<FCborStructSerializerBackend, FCborStructDeserializerBackend>();
 	}
-
-	return true;
 }
 
-IMPLEMENT_SIMPLE_AUTOMATION_TEST(FStructSerializerCborByteArrayTest, "System.Core.Serialization.StructSerializerCborByteArray", EAutomationTestFlags::ApplicationContextMask | EAutomationTestFlags::EngineFilter)
-
-bool FStructSerializerCborByteArrayTest::RunTest( const FString& Parameters )
+TEST_CASE_NAMED(FStructSerializerCborByteArrayTest, "System::Core::Serialization::StructSerializerCborByteArray", "[ApplicationContextMask][EngineFilter][StructSerializer]")
 {
 	// Ensure TArray<uint8>/TArray<int8> are written as CBOR byte string (~2x more compact) by default rather than a CBOR array.
 	{
@@ -882,7 +876,7 @@ bool FStructSerializerCborByteArrayTest::RunTest( const FString& Parameters )
 		FStructSerializerByteArray WrittenStruct;
 		FStructSerializer::Serialize(WrittenStruct, SerializerBackend);
 
-		TestTrue(TEXT("Arrays of int8/uint8 must be encoded in byte string (compact)"), Buffer.Num() == 54); // Copy the 54 bytes from VC++ Memory viewer to CBOR playground http://cbor.me/ to validate the count/content.
+		CHECK_MESSAGE(TEXT("Arrays of int8/uint8 must be encoded in byte string (compact)"), Buffer.Num() == 54); // Copy the 54 bytes from VC++ Memory viewer to CBOR playground http://cbor.me/ to validate the count/content.
 
 		// Deserialization
 		FMemoryReader Reader(Buffer);
@@ -892,11 +886,11 @@ bool FStructSerializerCborByteArrayTest::RunTest( const FString& Parameters )
 		FStructSerializerByteArray ReadStruct(NoInit);
 		FStructDeserializer::Deserialize(ReadStruct, DeserializerBackend, Policies);
 
-		TestTrue(TEXT("Value before TArray<uint8> must be the same before and after de-/serialization."), ReadStruct.Dummy1 == 1);
-		TestTrue(TEXT("Value after TArray<uint8> must be the same before and after de-/serialization."), ReadStruct.Dummy2 == 2);
-		TestTrue(TEXT("Value after TArray<int8> must be the same before and after de-/serialization."), ReadStruct.Dummy3 == 3);
-		TestTrue(TEXT("Array uint8 must be the same before and after de-/serialization"), WrittenStruct.ByteArray == ReadStruct.ByteArray);
-		TestTrue(TEXT("Array int8 must be the same before and after de-/serialization"), WrittenStruct.Int8Array == ReadStruct.Int8Array);
+		CHECK_MESSAGE(TEXT("Value before TArray<uint8> must be the same before and after de-/serialization."), ReadStruct.Dummy1 == 1);
+		CHECK_MESSAGE(TEXT("Value after TArray<uint8> must be the same before and after de-/serialization."), ReadStruct.Dummy2 == 2);
+		CHECK_MESSAGE(TEXT("Value after TArray<int8> must be the same before and after de-/serialization."), ReadStruct.Dummy3 == 3);
+		CHECK_MESSAGE(TEXT("Array uint8 must be the same before and after de-/serialization"), WrittenStruct.ByteArray == ReadStruct.ByteArray);
+		CHECK_MESSAGE(TEXT("Array int8 must be the same before and after de-/serialization"), WrittenStruct.Int8Array == ReadStruct.Int8Array);
 	}
 
 	// Ensure TArray<uint8>/TArray<int8> encoded in CBOR byte string are skipped on deserialization if required by the policy.
@@ -924,11 +918,11 @@ bool FStructSerializerCborByteArrayTest::RunTest( const FString& Parameters )
 		FStructSerializerByteArray ReadStruct(NoInit);
 		FStructDeserializer::Deserialize(ReadStruct, DeserializerBackend, Policies);
 
-		TestTrue(TEXT("Per deserializer policy, value before TArray<uint8> must be the same before and after de-/serialization."), ReadStruct.Dummy1 == 1);
-		TestTrue(TEXT("Per deserializer policy, value after TArray<uint8> must be the same before and after de-/serialization."), ReadStruct.Dummy2 == 2);
-		TestTrue(TEXT("Per deserializer policy, value after TArray<int8> must be the same before and after de-/serialization."), ReadStruct.Dummy3 == 3);
-		TestTrue(TEXT("Per deserializer policy, TArray<uint8> must be skipped on deserialization"), ReadStruct.ByteArray.Num() == 0);
-		TestTrue(TEXT("Per deserializer policy, TArray<int8> must be skipped on deserialization"), ReadStruct.Int8Array.Num() == 0);
+		CHECK_MESSAGE(TEXT("Per deserializer policy, value before TArray<uint8> must be the same before and after de-/serialization."), ReadStruct.Dummy1 == 1);
+		CHECK_MESSAGE(TEXT("Per deserializer policy, value after TArray<uint8> must be the same before and after de-/serialization."), ReadStruct.Dummy2 == 2);
+		CHECK_MESSAGE(TEXT("Per deserializer policy, value after TArray<int8> must be the same before and after de-/serialization."), ReadStruct.Dummy3 == 3);
+		CHECK_MESSAGE(TEXT("Per deserializer policy, TArray<uint8> must be skipped on deserialization"), ReadStruct.ByteArray.Num() == 0);
+		CHECK_MESSAGE(TEXT("Per deserializer policy, TArray<int8> must be skipped on deserialization"), ReadStruct.Int8Array.Num() == 0);
 	}
 
 	// Ensure empty TArray<uint8>/TArray<int8> are written as zero-length CBOR byte string.
@@ -943,7 +937,7 @@ bool FStructSerializerCborByteArrayTest::RunTest( const FString& Parameters )
 		WrittenStruct.Dummy3 = 3;
 		FStructSerializer::Serialize(WrittenStruct, SerializerBackend);
 
-		TestTrue(TEXT("Arrays of int8/uint8 must be encoded in byte string (compact)"), Buffer.Num() == 48); // Copy the 48 bytes from VC++ Memory viewer to CBOR playground http://cbor.me/ to validate the count/content.
+		CHECK_MESSAGE(TEXT("Arrays of int8/uint8 must be encoded in byte string (compact)"), Buffer.Num() == 48); // Copy the 48 bytes from VC++ Memory viewer to CBOR playground http://cbor.me/ to validate the count/content.
 
 		// Deserialization
 		FMemoryReader Reader(Buffer);
@@ -953,11 +947,11 @@ bool FStructSerializerCborByteArrayTest::RunTest( const FString& Parameters )
 		FStructSerializerByteArray ReadStruct(NoInit);
 		FStructDeserializer::Deserialize(ReadStruct, DeserializerBackend, Policies);
 
-		TestTrue(TEXT("Value before TArray<uint8> must be the same before and after de-/serialization."), ReadStruct.Dummy1 == 1);
-		TestTrue(TEXT("Value after TArray<uint8> must be the same before and after de-/serialization."), ReadStruct.Dummy2 == 2);
-		TestTrue(TEXT("Value after TArray<int8> must be the same before and after de-/serialization."), ReadStruct.Dummy3 == 3);
-		TestTrue(TEXT("Array uint8 must be the same before and after de-/serialization"), WrittenStruct.ByteArray == ReadStruct.ByteArray);
-		TestTrue(TEXT("Array int8 must be the same before and after de-/serialization"), WrittenStruct.Int8Array == ReadStruct.Int8Array);
+		CHECK_MESSAGE(TEXT("Value before TArray<uint8> must be the same before and after de-/serialization."), ReadStruct.Dummy1 == 1);
+		CHECK_MESSAGE(TEXT("Value after TArray<uint8> must be the same before and after de-/serialization."), ReadStruct.Dummy2 == 2);
+		CHECK_MESSAGE(TEXT("Value after TArray<int8> must be the same before and after de-/serialization."), ReadStruct.Dummy3 == 3);
+		CHECK_MESSAGE(TEXT("Array uint8 must be the same before and after de-/serialization"), WrittenStruct.ByteArray == ReadStruct.ByteArray);
+		CHECK_MESSAGE(TEXT("Array int8 must be the same before and after de-/serialization"), WrittenStruct.Int8Array == ReadStruct.Int8Array);
 	}
 
 	// Ensure TArray<uint8>/TArray<int8> CBOR serialization is backward compatible. (Serializer can write the old format and deserializer can read it)
@@ -971,7 +965,7 @@ bool FStructSerializerCborByteArrayTest::RunTest( const FString& Parameters )
 		FStructSerializerByteArray WrittenStruct;
 		FStructSerializer::Serialize(WrittenStruct, SerializerBackend);
 
-		TestTrue(TEXT("Backward compatibility: Serialized size check"), Buffer.Num() == 60); // Copy the 60 bytes from VC++ Memory viewer to CBOR playground http://cbor.me/ to validate the count/content.
+		CHECK_MESSAGE(TEXT("Backward compatibility: Serialized size check"), Buffer.Num() == 60); // Copy the 60 bytes from VC++ Memory viewer to CBOR playground http://cbor.me/ to validate the count/content.
 
 		// Deserialize TArray<uint8>/TArray<int8> as they were prior 4.25.
 		FMemoryReader Reader(Buffer);
@@ -981,15 +975,13 @@ bool FStructSerializerCborByteArrayTest::RunTest( const FString& Parameters )
 		FStructSerializerByteArray ReadStruct(NoInit);
 		FStructDeserializer::Deserialize(ReadStruct, DeserializerBackend, Policies);
 
-		TestTrue(TEXT("Backward compatibility: Integer must be the same before and after de-/serialization."), ReadStruct.Dummy1 == 1);
-		TestTrue(TEXT("Backward compatibility: Integer must be the same before and after de-/serialization."), ReadStruct.Dummy2 == 2);
-		TestTrue(TEXT("Backward compatibility: Integer must be the same before and after de-/serialization."), ReadStruct.Dummy3 == 3);
-		TestTrue(TEXT("Backward compatibility: TArray<uint8> must be readable as CBOR array of number."), WrittenStruct.ByteArray == ReadStruct.ByteArray);
-		TestTrue(TEXT("Backward compatibility: TArray<int8> must be readable as CBOR array of number."), WrittenStruct.Int8Array == ReadStruct.Int8Array);
+		CHECK_MESSAGE(TEXT("Backward compatibility: Integer must be the same before and after de-/serialization."), ReadStruct.Dummy1 == 1);
+		CHECK_MESSAGE(TEXT("Backward compatibility: Integer must be the same before and after de-/serialization."), ReadStruct.Dummy2 == 2);
+		CHECK_MESSAGE(TEXT("Backward compatibility: Integer must be the same before and after de-/serialization."), ReadStruct.Dummy3 == 3);
+		CHECK_MESSAGE(TEXT("Backward compatibility: TArray<uint8> must be readable as CBOR array of number."), WrittenStruct.ByteArray == ReadStruct.ByteArray);
+		CHECK_MESSAGE(TEXT("Backward compatibility: TArray<int8> must be readable as CBOR array of number."), WrittenStruct.Int8Array == ReadStruct.Int8Array);
 	}
-
-	return true;
 }
 
 
-#endif //WITH_DEV_AUTOMATION_TESTS
+#endif // WITH_TESTS
