Optional property support for StructSerializer & StructDeserializer, with tests.

#tests Unit tests pass

diff --git a/Engine/Source/Runtime/Serialization/Private/StructDeserializer.cpp b/Engine/Source/Runtime/Serialization/Private/StructDeserializer.cpp
--- a/Engine/Source/Runtime/Serialization/Private/StructDeserializer.cpp
+++ b/Engine/Source/Runtime/Serialization/Private/StructDeserializer.cpp
@@ -3,6 +3,7 @@
 #include "StructDeserializer.h"
 #include "UObject/UnrealType.h"
 #include "IStructDeserializerBackend.h"
+#include "UObject/PropertyOptional.h"
 #include "UObject/PropertyPortFlags.h"
 
 
@@ -259,6 +260,17 @@ bool FStructDeserializer::Deserialize( void* OutStruct, UStruct& TypeInfo, IStru
 						UE_LOG(LogSerialization, Verbose, TEXT("An item in map '%s' could not be read (%s)"), *PropertyName, *Backend.GetDebugString());
 					}
 				}
+				else if ((CurrentState.Property != nullptr) && (CurrentState.Property->GetClass() == FOptionalProperty::StaticClass()))
+				{
+					FOptionalProperty* OptionalProperty = CastField<FOptionalProperty>(CurrentState.Property);
+					FProperty* Property = OptionalProperty->GetValueProperty();
+					void* ValueData = OptionalProperty->MarkSetAndGetInitializedValuePointerToReplace(CurrentState.Data);
+
+					if (!Backend.ReadProperty(Property, CurrentState.Property, ValueData, CurrentState.ArrayIndex))
+					{
+						UE_LOG(LogSerialization, Verbose, TEXT("An item in optional '%s' could not be read (%s)"), *PropertyName, *Backend.GetDebugString());
+					}
+				}
 				else
 				{
 					// handle scalar property
@@ -364,6 +376,13 @@ bool FStructDeserializer::Deserialize( void* OutStruct, UStruct& TypeInfo, IStru
 
 					MapProperty->KeyProp->ImportText_Direct(*PropertyName, PairPtr, nullptr, PPF_None);
 				}
+				// handle map or struct element inside optional
+				else if ((CurrentState.Property != nullptr) && (CurrentState.Property->GetClass() == FOptionalProperty::StaticClass()))
+				{
+					FOptionalProperty* OptionalProperty = CastFieldChecked<FOptionalProperty>(CurrentState.Property);
+					NewState.Property = OptionalProperty->GetValueProperty();
+					NewState.Data = OptionalProperty->MarkSetAndGetInitializedValuePointerToReplace(CurrentState.Data);
+				}
 				else
 				{
 					NewState.Property = FindFProperty<FProperty>(CurrentState.TypeInfo, *PropertyName);
@@ -389,6 +408,12 @@ bool FStructDeserializer::Deserialize( void* OutStruct, UStruct& TypeInfo, IStru
 						FScriptMapHelper MapHelper(MapProperty, NewState.Data);
 						MapHelper.EmptyValues();
 					}
+					// handle optional property start
+					else if (FOptionalProperty* OptionalProperty = CastField<FOptionalProperty>(NewState.Property))
+					{
+						NewState.Data = OptionalProperty->ContainerPtrToValuePtr<void>(CurrentState.Data, CurrentState.ArrayIndex);
+						OptionalProperty->MarkUnset(NewState.Data);
+					}
 					// handle struct property
 					else
 					{
@@ -617,6 +642,17 @@ bool FStructDeserializer::DeserializeElement(void* OutAddress, UStruct& OwnerInf
 				
 				++CurrentState.ArrayIndex;
 			}
+			else if ((CurrentState.Property != nullptr) && (CurrentState.Property->GetClass() == FOptionalProperty::StaticClass()))
+			{
+				FOptionalProperty* OptionalProperty = CastField<FOptionalProperty>(CurrentState.Property);
+				FProperty* Property = OptionalProperty->GetValueProperty();
+				void* ValueData = OptionalProperty->MarkSetAndGetInitializedValuePointerToReplace(CurrentState.Data);
+
+				if (!Backend.ReadProperty(Property, CurrentState.Property, ValueData, CurrentState.ArrayIndex))
+				{
+					UE_LOG(LogSerialization, Verbose, TEXT("An item in optional '%s' could not be read (%s)"), *PropertyName, *Backend.GetDebugString());
+				}
+			}
 			// Otherwise we are dealing with dynamic or static array
 			else if (PropertyName.IsEmpty())
 			{
@@ -774,6 +810,26 @@ bool FStructDeserializer::DeserializeElement(void* OutAddress, UStruct& OwnerInf
 							continue;
 						}
 					}
+					else if (FOptionalProperty* OptionalProperty = CastField<FOptionalProperty>(Property))
+					{
+						void* Data = OptionalProperty->ContainerPtrToValuePtr<void>(CurrentState.Data);
+
+						if (void* ValueData = OptionalProperty->GetValuePointerForReadOrReplaceIfSet(Data))
+						{
+							Property = OptionalProperty->GetValueProperty();
+
+							//Offset the pointer directly and give index 0 to be read so no offsetting is done during deserialization
+							CurrentState.Data = ValueData;
+							CurrentState.ArrayIndex = 0;
+						}
+						else
+						{
+							// Not set
+							UE_LOG(LogSerialization, Verbose, TEXT("TOptional %s is not set and is trying to be read"), *OptionalProperty->GetFName().ToString());
+							Backend.SkipStructure();
+							continue;
+						}
+					}
 
 					if (!Backend.ReadProperty(Property, nullptr, CurrentState.Data, CurrentState.ArrayIndex))
 					{
@@ -894,6 +950,14 @@ bool FStructDeserializer::DeserializeElement(void* OutAddress, UStruct& OwnerInf
 						continue;
 					}
 				}
+				// handle struct element inside optional
+				else if (FOptionalProperty* OptionalProperty = CastField<FOptionalProperty>(CurrentState.Property))
+				{
+					void* Data = OptionalProperty->ContainerPtrToValuePtr<void>(CurrentState.Data);
+					NewState.Property = OptionalProperty->GetValueProperty();
+					NewState.Data = OptionalProperty->MarkSetAndGetInitializedValuePointerToReplace(Data);
+					NewState.ArrayIndex = 0;
+				}
 				else
 				{
 					//Property was found so we might be in a static array of struct
@@ -913,6 +977,13 @@ bool FStructDeserializer::DeserializeElement(void* OutAddress, UStruct& OwnerInf
 					}
 				}
 			}
+			// handle map or struct element inside optional
+			else if ((CurrentState.Property != nullptr) && (CurrentState.Property->GetClass() == FOptionalProperty::StaticClass()))
+			{
+				FOptionalProperty* OptionalProperty = CastFieldChecked<FOptionalProperty>(CurrentState.Property);
+				NewState.Property = OptionalProperty->GetValueProperty();
+				NewState.Data = OptionalProperty->MarkSetAndGetInitializedValuePointerToReplace(CurrentState.Data);
+			}
 			else
 			{
 				NewState.Property = FindFProperty<FProperty>(CurrentState.TypeInfo, *PropertyName);
@@ -1009,6 +1080,12 @@ bool FStructDeserializer::DeserializeElement(void* OutAddress, UStruct& OwnerInf
 						}
 					}
 				}
+				// handle optional property entry
+				else if (FOptionalProperty* OptionalProperty = CastField<FOptionalProperty>(NewState.Property))
+				{
+					NewState.Data = OptionalProperty->ContainerPtrToValuePtr<void>(CurrentState.Data, CurrentState.ArrayIndex);
+					OptionalProperty->MarkUnset(NewState.Data);
+				}
 				// handle struct property
 				else
 				{
diff --git a/Engine/Source/Runtime/Serialization/Private/StructSerializer.cpp b/Engine/Source/Runtime/Serialization/Private/StructSerializer.cpp
--- a/Engine/Source/Runtime/Serialization/Private/StructSerializer.cpp
+++ b/Engine/Source/Runtime/Serialization/Private/StructSerializer.cpp
@@ -2,6 +2,7 @@
 
 #include "StructSerializer.h"
 #include "UObject/UnrealType.h"
+#include "UObject/PropertyOptional.h"
 #include "IStructSerializerBackend.h"
 
 
@@ -143,6 +144,33 @@ namespace StructSerializer
 		}
 	}
 
+	void SerializeOptional(FStructSerializerState& CurrentState, TArray<FStructSerializerState>& StateStack, IStructSerializerBackend& Backend)
+	{
+		if (!CurrentState.HasBeenProcessed)
+		{
+			Backend.BeginStructure(CurrentState);
+
+			CurrentState.HasBeenProcessed = true;
+			StateStack.Push(CurrentState);
+
+			const FOptionalProperty* OptionalProperty = CastFieldChecked<FOptionalProperty>(CurrentState.ValueProperty);
+			FProperty* InnerProperty = OptionalProperty->GetValueProperty();
+			const void* ValueData = CurrentState.ValueProperty->ContainerPtrToValuePtr<void>(CurrentState.ValueData);
+			if (const void* InnerValue = OptionalProperty->GetValuePointerForReadIfSet(ValueData))
+			{
+				FStructSerializerState NewState;
+				NewState.ValueData = InnerValue;
+				NewState.ValueProperty = InnerProperty;
+				NewState.FieldType = InnerProperty->GetClass();
+				StateStack.Add(NewState);
+			}
+		}
+		else
+		{
+			Backend.EndStructure(CurrentState);
+		}
+	}
+
 	void SerializeStaticArray(FStructSerializerState& CurrentState, IStructSerializerBackend& Backend)
 	{
 		if (CurrentState.ElementIndex != INDEX_NONE)
@@ -343,6 +371,12 @@ namespace StructSerializer
 				SerializeIterable<FSetProperty>(CurrentState, StateStack, Backend, Policies);
 			}
 
+			// Optionals
+			else if (CastField<FOptionalProperty>(CurrentState.ValueProperty))
+			{
+				SerializeOptional(CurrentState, StateStack, Backend);
+			}
+
 			// Static arrays
 			else if (CurrentState.ValueProperty->ArrayDim > 1)
 			{
diff --git a/Engine/Source/Runtime/Serialization/Private/Tests/StructSerializerTest.cpp b/Engine/Source/Runtime/Serialization/Private/Tests/StructSerializerTest.cpp
--- a/Engine/Source/Runtime/Serialization/Private/Tests/StructSerializerTest.cpp
+++ b/Engine/Source/Runtime/Serialization/Private/Tests/StructSerializerTest.cpp
@@ -146,6 +146,18 @@ namespace StructSerializerTest
 		CHECK_MESSAGE(TEXT("Sets.StructSet must be the same before and after de-/serialization"), Struct1.StructSet.Num() == Struct2.StructSet.Num() && Struct1.StructSet.Difference(Struct2.StructSet).Num() == 0);
 	}
 
+	void ValidateOptionals(const FStructSerializerOptionalTestStruct& Struct1, const FStructSerializerOptionalTestStruct& Struct2)
+	{
+		CHECK_EQUALS(TEXT("Optionals.StrOptional must be the same before and after de-/serialization"), Struct1.StrOptional, Struct2.StrOptional);
+		CHECK_EQUALS(TEXT("Optionals.StrOptionalUnset must be the same before and after de-/serialization"), Struct1.StrOptionalUnset, Struct2.StrOptionalUnset);
+		CHECK_EQUALS(TEXT("Optionals.IntOptional must be the same before and after de-/serialization"), Struct1.IntOptional, Struct2.IntOptional);
+		CHECK_EQUALS(TEXT("Optionals.IntOptionalUnset must be the same before and after de-/serialization"), Struct1.IntOptionalUnset, Struct2.IntOptionalUnset);
+		CHECK_EQUALS(TEXT("Optionals.NameOptional must be the same before and after de-/serialization"), Struct1.NameOptional, Struct2.NameOptional);
+		CHECK_EQUALS(TEXT("Optionals.NameOptionalUnset must be the same before and after de-/serialization"), Struct1.NameOptionalUnset, Struct2.NameOptionalUnset);
+		CHECK_EQUALS(TEXT("Optionals.StructOptional must be the same before and after de-/serialization"), Struct1.StructOptional, Struct2.StructOptional);
+		CHECK_EQUALS(TEXT("Optionals.StructOptionalUnset must be the same before and after de-/serialization"), Struct1.StructOptionalUnset, Struct2.StructOptionalUnset);
+	}
+
 	void ValidateLWCSerializationBackwardCompatibility(const FStructSerializerLWCTypesTest& Struct1, const FStructSerializerNonLWCTypesTest& Struct2)
 	{
 		//Make comparison by casting the double (lwc) version down to float (non-lwc) since this is what will happen during serialization
@@ -503,6 +515,39 @@ namespace StructSerializerTest
 				ValidateSets(TestStruct.Sets, TestStruct2.Sets);
 			}
 
+			//Optionals
+			{
+				TArray<uint8> Buffer;
+				FMemoryReader Reader(Buffer);
+				FMemoryWriter Writer(Buffer);
+
+				TSerializerBackend SerializerBackend(Writer, EStructSerializerBackendFlags::Default);
+				TDeserializerBackend DeserializerBackend(Reader);
+
+				FStructSerializerTestStruct TestStruct = OriginalStruct;
+				const FName Member = GET_MEMBER_NAME_CHECKED(FStructSerializerTestStruct, Optionals);
+				FProperty* Property = FindFProperty<FProperty>(FStructSerializerTestStruct::StaticStruct(), Member);
+				FStructSerializer::SerializeElement(&TestStruct, Property, INDEX_NONE, SerializerBackend, Policies);
+
+				// Inputs:
+				// - TestStruct: all the properties will be defaulted - "Unset" properties unset, and the rest set.
+				// - TestStruct2: all the "Unset" properties will be set, and the rest will be unset.
+				// Expected Output:
+				// - TestStruct2: all the "Unset" properties will be unset, and the rest will be set.
+				// Goal:
+				// - Validate that unset properties are populated when data is available for deserialization.
+				// - Validate that set properties are cleared when data isn't available for deserialization.
+				FStructSerializerTestStruct TestStruct2(NoInit);
+				TestStruct2.Optionals.StrOptionalUnset = TestStruct.Optionals.StrOptional;
+				TestStruct2.Optionals.IntOptionalUnset = TestStruct.Optionals.IntOptional;
+				TestStruct2.Optionals.NameOptionalUnset = TestStruct.Optionals.NameOptional;
+				TestStruct2.Optionals.StructOptionalUnset = TestStruct.Optionals.StructOptional;
+
+				CHECK_MESSAGE(TEXT("Deserialization must succeed"), FStructDeserializer::DeserializeElement(&TestStruct2, *FStructSerializerTestStruct::StaticStruct(), INDEX_NONE, DeserializerBackend, DeserializerPolicies));
+
+				ValidateOptionals(TestStruct.Optionals, TestStruct2.Optionals);
+			}
+
 			//TArray<uint8> element
 			{
 				TArray<uint8> Buffer;
@@ -737,6 +782,9 @@ namespace StructSerializerTest
 		// test sets
 		ValidateSets(TestStruct.Sets, TestStruct2.Sets);
 
+		// test optionals
+		ValidateOptionals(TestStruct.Optionals, TestStruct2.Optionals);
+
 		//Test LWC types with standard de-serialization
 		ValidateLWCTypes(TestStruct.LWCTypes, TestStruct2.LWCTypes);
 	}
diff --git a/Engine/Source/Runtime/Serialization/Private/Tests/StructSerializerTestTypes.h b/Engine/Source/Runtime/Serialization/Private/Tests/StructSerializerTestTypes.h
--- a/Engine/Source/Runtime/Serialization/Private/Tests/StructSerializerTestTypes.h
+++ b/Engine/Source/Runtime/Serialization/Private/Tests/StructSerializerTestTypes.h
@@ -689,7 +689,64 @@ struct FStructSerializerSetTestStruct
 	FStructSerializerSetTestStruct( ENoInit ) { }
 };
 
+/**
+ * Test structure for optional properties.
+ */
+USTRUCT()
+struct FStructSerializerOptionalTestStruct
+{
+	GENERATED_BODY()
+
+	UPROPERTY()
+	TOptional<FString> StrOptional;
+
+	UPROPERTY()
+	TOptional<FString> StrOptionalUnset;
+
+	UPROPERTY()
+	TOptional<int32> IntOptional;
+
+	UPROPERTY()
+	TOptional<int32> IntOptionalUnset;
+
+	UPROPERTY()
+	TOptional<FName> NameOptional;
+
+	UPROPERTY()
+	TOptional<FName> NameOptionalUnset;
 
+	UPROPERTY(meta=(IgnoreForMemberInitializationTest))
+	TOptional<FStructSerializerBuiltinTestStruct> StructOptional;
+
+	UPROPERTY()
+	TOptional<FStructSerializerBuiltinTestStruct> StructOptionalUnset;
+
+	// At the time of writing, optionals can not be used as keys/values in containers (optionals/arrays/sets/maps).
+
+	// UPROPERTY()
+	// TOptional<TOptional<int32>> OptionalOptionalInt;
+
+	// UPROPERTY()
+	// TArray<TOptional<int32>> ArrayOptionalInt;
+
+	// UPROPERTY()
+	// TSet<TOptional<int32>> SetOptionalInt;
+
+	// UPROPERTY()
+	// TMap<TOptional<int32>, TOptional<int32>> MapOptionalIntOptionalInt;
+
+	/** Default constructor. */
+	FStructSerializerOptionalTestStruct()
+		: StrOptional(TEXT("Test String"))
+		, IntOptional(1234)
+		, NameOptional(TEXT("Test Name"))
+		, StructOptional(FStructSerializerBuiltinTestStruct())
+	{
+	}
+
+	/** Creates an uninitialized instance. */
+	FStructSerializerOptionalTestStruct( ENoInit ) { }
+};
 
 /**
  * Test structure for all supported types.
@@ -720,6 +777,9 @@ struct FStructSerializerTestStruct
 	UPROPERTY()
 	FStructSerializerSetTestStruct Sets;
 
+	UPROPERTY()
+	FStructSerializerOptionalTestStruct Optionals;
+
 	UPROPERTY()
 	FStructSerializerLWCTypesTest LWCTypes;
 
@@ -735,6 +795,7 @@ struct FStructSerializerTestStruct
 		, Arrays(NoInit)
 		, Maps(NoInit)
 		, Sets(NoInit)
+		, Optionals(NoInit)
 		, LWCTypes(NoInit)
 	{ }
 };
