Json: Reduced the probability of OOM crashes when running UE.Automation tests on some memory-constrained platforms.

The crashes were caused by us not being able to allocate large memory buffers. Plenty of physical memory was available but virtual memory was too fragmented. This fragmentation was caused by our tests repeatedly (after each rendering test) allocating very large chunks of memory i.e. ~70 MB when sending image data back to the connected PC coordinating the tests. This 70 MB was needed because of very inefficient serialization (FStructSerializer), noticeable in particular when handling raw binary buffers (TArray<int8>).

The core issue is that the serializer does the job in two passes. First, it creates descriptions of all collection items (FStructSerializerState), which are added to a common stack, and then iterating over all the stack elements. It means that each byte from a serialized binary buffer, we would allocate one FStructSerializerState (64 bytes). That's a 64x increase!

Reworking the whole FStructSerializer would be a pretty large change, so instead I'm implementing a workaround specifically for binary buffer. It implements the existing serialization logic in a slightly different way as a WritePODArray special case. We simply directly iterate over binary arrays and skip the state generation step completelly.

diff --git a/Engine/Source/Runtime/Serialization/Private/Backends/JsonStructSerializerBackend.cpp b/Engine/Source/Runtime/Serialization/Private/Backends/JsonStructSerializerBackend.cpp
--- a/Engine/Source/Runtime/Serialization/Private/Backends/JsonStructSerializerBackend.cpp
+++ b/Engine/Source/Runtime/Serialization/Private/Backends/JsonStructSerializerBackend.cpp
@@ -203,3 +203,40 @@ void FJsonStructSerializerBackend::WriteProperty(const FStructSerializerState& S
 		UE_LOG(LogSerialization, Verbose, TEXT("FJsonStructSerializerBackend: Property %s cannot be serialized, because its type (%s) is not supported"), *State.ValueProperty->GetFName().ToString(), *State.ValueType->GetFName().ToString());
 	}
 }
+
+bool FJsonStructSerializerBackend::WritePODArray(const FStructSerializerState& State)
+{
+	if (State.ElementIndex != INDEX_NONE)
+	{
+		return false;
+	}
+
+	// This code serializes TArray<uint8> in the same way FStructSerializer would do. However, we iterate over
+	// elements directly while FStructSerializer allocates one FStructSerializerState (64 bytes) for each
+	// byte item. As a result, one 1 MB of binary data would require more than 64 MB of temporary state data.
+	FArrayProperty* ArrayProperty = CastField<FArrayProperty>(State.ValueProperty);
+	if (ArrayProperty && (CastField<FByteProperty>(ArrayProperty->Inner) || CastField<FInt8Property>(ArrayProperty->Inner)))
+	{
+		FScriptArrayHelper ArrayHelper(ArrayProperty, ArrayProperty->ContainerPtrToValuePtr<void>(State.ValueData));
+
+		FStructSerializerState InnerState(nullptr, ArrayProperty->Inner, EStructSerializerStateFlags::None);
+
+		for (int32 Index = 0; Index < ArrayHelper.Num(); ++Index)
+		{
+			if (ArrayHelper.IsValidIndex(Index))
+			{
+				InnerState.ValueData = ArrayHelper.GetRawPtr(Index);
+
+				WriteProperty(InnerState);
+			}
+		}
+
+		// We need to close the array ourselves because FStructSerializer doesn't do this after we've declared we've completed
+		// serialization by returning true.
+		EndArray(State);
+
+		return true;
+	}
+
+	return false;
+}
diff --git a/Engine/Source/Runtime/Serialization/Public/Backends/JsonStructSerializerBackend.h b/Engine/Source/Runtime/Serialization/Public/Backends/JsonStructSerializerBackend.h
--- a/Engine/Source/Runtime/Serialization/Public/Backends/JsonStructSerializerBackend.h
+++ b/Engine/Source/Runtime/Serialization/Public/Backends/JsonStructSerializerBackend.h
@@ -60,6 +60,7 @@ public:
 	SERIALIZATION_API virtual void EndStructure(const FStructSerializerState& State) override;
 	SERIALIZATION_API virtual void WriteComment(const FString& Comment) override;
 	SERIALIZATION_API virtual void WriteProperty(const FStructSerializerState& State, int32 ArrayIndex = 0) override;
+	SERIALIZATION_API virtual bool WritePODArray(const FStructSerializerState& State) override;
 
 protected:
 
