Struct Serializer: deserializing now grows the containers to fit incoming entries rather than skipping entries out of range entirely.

diff --git a/Engine/Source/Runtime/Serialization/Private/StructDeserializer.cpp b/Engine/Source/Runtime/Serialization/Private/StructDeserializer.cpp
--- a/Engine/Source/Runtime/Serialization/Private/StructDeserializer.cpp
+++ b/Engine/Source/Runtime/Serialization/Private/StructDeserializer.cpp
@@ -70,6 +70,30 @@ namespace StructDeserializer
 
 		return Class;
 	}
+
+	/** Finds an element in a Map/Set Container at the given logical index, adding enough entries for that logical index to be valid */
+	template<typename ScriptContainerHelperType>
+	int32 ExpandForIndex(ScriptContainerHelperType& ScriptContainerHelper, int32 LogicalIndex)
+	{
+		int32 InternalIndex = ScriptContainerHelper.FindInternalIndex(LogicalIndex);
+		if (InternalIndex != INDEX_NONE)
+		{
+			return InternalIndex;
+		}
+
+		// Container does not have any element at the logical array index, add enough items so that this logical index is valid
+		const int32 ItemsToAddCount = LogicalIndex - ScriptContainerHelper.Num() + 1;
+
+		// Maps/Set containers don't have a ExpandForIndex / Add(Count). Add default values manually
+		for (int32 ItemsAddedCount = 0; ItemsAddedCount < ItemsToAddCount; ++ItemsAddedCount)
+		{
+			ScriptContainerHelper.AddDefaultValue_Invalid_NeedsRehash();
+		}
+
+		InternalIndex = ScriptContainerHelper.FindInternalIndex(LogicalIndex);
+		check(InternalIndex != INDEX_NONE);
+		return InternalIndex;
+	}
 }
 
 
@@ -595,22 +619,15 @@ bool FStructDeserializer::DeserializeElement(void* OutAddress, UStruct& OwnerInf
 				// handle set element
 				FSetProperty* SetProperty = CastField<FSetProperty>(CurrentState.Property);
 				FScriptSetHelper SetHelper(SetProperty, CurrentState.Data);
-				const int32 InternalIndex = SetHelper.FindInternalIndex(CurrentState.ArrayIndex);
-				if (InternalIndex != INDEX_NONE)
-				{
-					uint8* ElementPtr = SetHelper.GetElementPtr(InternalIndex);
-					FProperty* Property = SetProperty->ElementProp;
-					constexpr int32 ReadIndex = 0; //Pointer is offset so reading index is 0
-					if (!Backend.ReadProperty(Property, CurrentState.Property, ElementPtr, ReadIndex))
-					{
-						UE_LOG(LogSerialization, Verbose, TEXT("An item in Set '%s' could not be read (%s)"), *PropertyName, *Backend.GetDebugString());
-					}
-				}
-				else
+
+				const int32 InternalIndex = ExpandForIndex(SetHelper, CurrentState.ArrayIndex);
+
+				uint8* ElementPtr = SetHelper.GetElementPtr(InternalIndex);
+				FProperty* Property = SetProperty->ElementProp;
+				constexpr int32 ReadIndex = 0; //Pointer is offset so reading index is 0
+				if (!Backend.ReadProperty(Property, CurrentState.Property, ElementPtr, ReadIndex))
 				{
-					//Too many entries for given Set
-					UE_LOG(LogSerialization, Verbose, TEXT("TSet %s has dimension of %d and trying to read element %d"), *CurrentState.Property->GetFName().ToString(), SetHelper.Num(), CurrentState.ArrayIndex);
-					continue;
+					UE_LOG(LogSerialization, Verbose, TEXT("An item in Set '%s' could not be read (%s)"), *PropertyName, *Backend.GetDebugString());
 				}
 				++CurrentState.ArrayIndex;
 			}
@@ -622,22 +639,14 @@ bool FStructDeserializer::DeserializeElement(void* OutAddress, UStruct& OwnerInf
 				FScriptMapHelper MapHelper(MapProperty, CurrentState.Data);
 
 				//When written as array, maps won't include the key, only values
-				const int32 InternalIndex = MapHelper.FindInternalIndex(CurrentState.ArrayIndex);
-				if (InternalIndex != INDEX_NONE)
-				{
-					uint8* PairPtr = MapHelper.GetPairPtr(InternalIndex);
-					FProperty* Property = MapProperty->ValueProp;
-					constexpr int32 ReadIndex = 0; //Pointer is offset so reading index is 0
-					if (!Backend.ReadProperty(Property, CurrentState.Property, PairPtr, ReadIndex))
-					{
-						UE_LOG(LogSerialization, Verbose, TEXT("An item in Set '%s' could not be read (%s)"), *PropertyName, *Backend.GetDebugString());
-					}
-				}
-				else
+				const int32 InternalIndex = ExpandForIndex(MapHelper, CurrentState.ArrayIndex);
+
+				uint8* PairPtr = MapHelper.GetPairPtr(InternalIndex);
+				FProperty* Property = MapProperty->ValueProp;
+				constexpr int32 ReadIndex = 0; //Pointer is offset so reading index is 0
+				if (!Backend.ReadProperty(Property, CurrentState.Property, PairPtr, ReadIndex))
 				{
-					//Too many entries for given Map
-					UE_LOG(LogSerialization, Verbose, TEXT("TMap %s has dimension of %d and trying to read element %d"), *CurrentState.Property->GetFName().ToString(), MapHelper.Num(), CurrentState.ArrayIndex);
-					continue;
+					UE_LOG(LogSerialization, Verbose, TEXT("An item in Set '%s' could not be read (%s)"), *PropertyName, *Backend.GetDebugString());
 				}
 				
 				++CurrentState.ArrayIndex;
@@ -666,21 +675,13 @@ bool FStructDeserializer::DeserializeElement(void* OutAddress, UStruct& OwnerInf
 				{
 					// dynamic array element
 					FScriptArrayHelper ArrayHelper(ArrayProperty, ArrayProperty->ContainerPtrToValuePtr<void>(CurrentState.Data));
-					if (ArrayHelper.IsValidIndex(CurrentState.ArrayIndex))
-					{
-						Property = ArrayProperty->Inner;
-						DataAddress = ArrayHelper.GetRawPtr(CurrentState.ArrayIndex);
-						
-						//Arraydim will be 1 for inner TArray properties. Offset the read data and keep index at 0
-						CurrentArrayIndex = 0; 
-					}
-					else
-					{
-						checkSlow(CurrentState.Property);
-						//Too many entries in TArray
-						UE_LOG(LogSerialization, Verbose, TEXT("TArray %s has dimension of %d and trying to read element %d"), *CurrentState.Property->GetFName().ToString(), CurrentState.Property->ArrayDim, CurrentState.ArrayIndex);
-						continue;
-					}
+					ArrayHelper.ExpandForIndex(CurrentState.ArrayIndex);
+
+					Property = ArrayProperty->Inner;
+					DataAddress = ArrayHelper.GetRawPtr(CurrentState.ArrayIndex);
+
+					//Arraydim will be 1 for inner TArray properties. Offset the read data and keep index at 0
+					CurrentArrayIndex = 0; 
 				}
 				else
 				{
@@ -733,84 +734,54 @@ bool FStructDeserializer::DeserializeElement(void* OutAddress, UStruct& OwnerInf
 					{
 						FScriptSetHelper SetHelper(SetProperty, SetProperty->ContainerPtrToValuePtr<void>(CurrentState.Data));
 
-						//If a specific index is asked and it's not valid, skip the property
-						const int32 InternalIndex = SetHelper.FindInternalIndex(CurrentState.ArrayIndex);
-						if (InternalIndex != INDEX_NONE)
-						{
-							Property = SetProperty->ElementProp;
+						const int32 InternalIndex = ExpandForIndex(SetHelper, CurrentState.ArrayIndex);
 
-							//Offset the pointer directly and give index 0 to be read so no offsetting is done during deserialization
-							CurrentState.Data = SetHelper.GetElementPtr(InternalIndex);
-							CurrentState.ArrayIndex = 0;
+						Property = SetProperty->ElementProp;
 
-							if (!Backend.ReadProperty(Property, nullptr, CurrentState.Data, CurrentState.ArrayIndex))
-							{
-								UE_LOG(LogSerialization, Verbose, TEXT("The property '%s' could not be read (%s)"), *PropertyName, *Backend.GetDebugString());
-							}
+						//Offset the pointer directly and give index 0 to be read so no offsetting is done during deserialization
+						CurrentState.Data = SetHelper.GetElementPtr(InternalIndex);
+						CurrentState.ArrayIndex = 0;
 
-							//On element of a set was written so rehash it
-							SetHelper.Rehash();
-							continue;
-						}
-						else
+						if (!Backend.ReadProperty(Property, nullptr, CurrentState.Data, CurrentState.ArrayIndex))
 						{
-							//Index out of bound
-							UE_LOG(LogSerialization, Verbose, TEXT("Set %s has dimension of %d and trying to read element %d"), *SetProperty->GetFName().ToString(), SetHelper.Num(), CurrentState.ArrayIndex);
-							Backend.SkipStructure();
-							continue;
+							UE_LOG(LogSerialization, Verbose, TEXT("The property '%s' could not be read (%s)"), *PropertyName, *Backend.GetDebugString());
 						}
+
+						//An element of a set was written so rehash it
+						SetHelper.Rehash();
+						continue;
 					}
 					else if (FMapProperty* MapProperty = CastField<FMapProperty>(Property))
 					{
 						FScriptMapHelper MapHelper(MapProperty, MapProperty->ContainerPtrToValuePtr<void>(CurrentState.Data));
 
-						//If a specific index is asked and it's not valid, skip the property
-						const int32 InternalIndex = MapHelper.FindInternalIndex(CurrentState.ArrayIndex);
-						if (InternalIndex != INDEX_NONE)
-						{
-							Property = MapProperty->ValueProp;
+						const int32 InternalIndex = ExpandForIndex(MapHelper, CurrentState.ArrayIndex);
 
-							//Offset the pointer directly and give index 0 to be read so no offsetting is done during deserialization
-							CurrentState.Data = MapHelper.GetPairPtr(InternalIndex);
-							CurrentState.ArrayIndex = 0;
+						Property = MapProperty->ValueProp;
 
-							if (!Backend.ReadProperty(Property, nullptr, CurrentState.Data, CurrentState.ArrayIndex))
-							{
-								UE_LOG(LogSerialization, Verbose, TEXT("The property '%s' could not be read (%s)"), *PropertyName, *Backend.GetDebugString());
-							}
+						//Offset the pointer directly and give index 0 to be read so no offsetting is done during deserialization
+						CurrentState.Data = MapHelper.GetPairPtr(InternalIndex);
+						CurrentState.ArrayIndex = 0;
 
-							//On element of a map was written so rehash it
-							MapHelper.Rehash();
-							continue;
-						}
-						else
+						if (!Backend.ReadProperty(Property, nullptr, CurrentState.Data, CurrentState.ArrayIndex))
 						{
-							//Index out of bound
-							UE_LOG(LogSerialization, Verbose, TEXT("TMap %s has dimension of %d and trying to read element %d"), *MapProperty->GetFName().ToString(), MapHelper.Num(), CurrentState.ArrayIndex);
-							Backend.SkipStructure();
-							continue;
+							UE_LOG(LogSerialization, Verbose, TEXT("The property '%s' could not be read (%s)"), *PropertyName, *Backend.GetDebugString());
 						}
+
+						//An element of a map was written so rehash it
+						MapHelper.Rehash();
+						continue;
 					}
 					else if (FArrayProperty* ArrayProperty = CastField<FArrayProperty>(Property))
 					{
 						FScriptArrayHelper ArrayHelper(ArrayProperty, ArrayProperty->ContainerPtrToValuePtr<void>(CurrentState.Data));
+						ArrayHelper.ExpandForIndex(CurrentState.ArrayIndex);
 
-						//If a specific index is asked and it's not valid, skip the property
-						if (ArrayHelper.IsValidIndex(CurrentState.ArrayIndex))
-						{
-							Property = ArrayProperty->Inner;
+						Property = ArrayProperty->Inner;
 
-							//Offset the pointer directly and give index 0 to be read so no offsetting is done during deserialization
-							CurrentState.Data = ArrayHelper.GetRawPtr(CurrentState.ArrayIndex);
-							CurrentState.ArrayIndex = 0;
-						}
-						else
-						{
-							//Index out of bound
-							UE_LOG(LogSerialization, Verbose, TEXT("TArray %s has dimension of %d and trying to read element %d"), *ArrayProperty->GetFName().ToString(), ArrayHelper.Num(), CurrentState.ArrayIndex);
-							Backend.SkipStructure();
-							continue;
-						}
+						//Offset the pointer directly and give index 0 to be read so no offsetting is done during deserialization
+						CurrentState.Data = ArrayHelper.GetRawPtr(CurrentState.ArrayIndex);
+						CurrentState.ArrayIndex = 0;
 					}
 					else if (FOptionalProperty* OptionalProperty = CastField<FOptionalProperty>(Property))
 					{
@@ -860,12 +831,12 @@ bool FStructDeserializer::DeserializeElement(void* OutAddress, UStruct& OwnerInf
 			// rehash if value was a map, set
 			if(FMapProperty* MapProperty = CastField<FMapProperty>(CurrentState.Property))
 			{
-				FScriptMapHelper MapHelper(MapProperty, MapProperty->ContainerPtrToValuePtr<void>(CurrentState.Data));
+				FScriptMapHelper MapHelper(MapProperty, CurrentState.Data);
 				MapHelper.Rehash();
 			}
 			else if (FSetProperty* SetProperty = CastField<FSetProperty>(CurrentState.Property))
 			{
-				FScriptSetHelper SetHelper(SetProperty, SetProperty->ContainerPtrToValuePtr<void>(CurrentState.Data));
+				FScriptSetHelper SetHelper(SetProperty, CurrentState.Data);
 				SetHelper.Rehash();
 			}
 
@@ -896,66 +867,42 @@ bool FStructDeserializer::DeserializeElement(void* OutAddress, UStruct& OwnerInf
 				if (FSetProperty* SetProperty = CastField<FSetProperty>(CurrentState.Property))
 				{
 					FScriptSetHelper SetHelper(SetProperty, CurrentState.Data);
-					const int32 InternalIndex = SetHelper.FindInternalIndex(CurrentState.ArrayIndex);
-					if (InternalIndex != INDEX_NONE)
-					{
-						NewState.Property = SetProperty->ElementProp;
-						NewState.Data = SetHelper.GetElementPtr(InternalIndex);
-						NewState.ArrayIndex = 0;
-						++CurrentState.ArrayIndex;
-					}
-					else
-					{
-						//Too many entries
-						UE_LOG(LogSerialization, Verbose, TEXT("TSet %s has dimension of %d and trying to read element %d"), *CurrentState.Property->GetFName().ToString(), SetHelper.Num(), CurrentState.ArrayIndex);
-						Backend.SkipStructure();
-						continue;
-					}
+
+					const int32 InternalIndex = ExpandForIndex(SetHelper, CurrentState.ArrayIndex);
+
+					NewState.Property = SetProperty->ElementProp;
+					NewState.Data = SetHelper.GetElementPtr(InternalIndex);
+					NewState.ArrayIndex = 0;
+					++CurrentState.ArrayIndex;
 				}
 				else if (FMapProperty* MapProperty = CastField<FMapProperty>(CurrentState.Property))
 				{
 					FScriptMapHelper MapHelper(MapProperty, CurrentState.Data);
-					const int32 InternalIndex = MapHelper.FindInternalIndex(CurrentState.ArrayIndex);
-					if (InternalIndex != INDEX_NONE)
-					{
-						NewState.Property = MapProperty->ValueProp;
-						NewState.Data = MapHelper.GetValuePtr(InternalIndex);
-						NewState.ArrayIndex = 0;
-						++CurrentState.ArrayIndex;
-					}
-					else
-					{
-						//Too many entries
-						UE_LOG(LogSerialization, Verbose, TEXT("TMap %s has dimension of %d and trying to read element %d"), *CurrentState.Property->GetFName().ToString(), MapHelper.Num(), CurrentState.ArrayIndex);
-						Backend.SkipStructure();
-						continue;
-					}
+
+					const int32 InternalIndex = ExpandForIndex(MapHelper, CurrentState.ArrayIndex);
+
+					NewState.Property = MapProperty->ValueProp;
+					NewState.Data = MapHelper.GetValuePtr(InternalIndex);
+					NewState.ArrayIndex = 0;
+					++CurrentState.ArrayIndex;
 				}
 				// handle struct element inside array
 				else if (FArrayProperty* ArrayProperty = CastField<FArrayProperty>(CurrentState.Property))
 				{
 					FScriptArrayHelper ArrayHelper(ArrayProperty, ArrayProperty->ContainerPtrToValuePtr<void>(CurrentState.Data));
 
-					//We're going over the array, make sure indexing is valid. Could add support for growing arrays
-					if (ArrayHelper.IsValidIndex(CurrentState.ArrayIndex))
-					{
-						NewState.Property = ArrayProperty->Inner;
-						NewState.Data = ArrayHelper.GetRawPtr(CurrentState.ArrayIndex);
-						NewState.ArrayIndex = 0;
-						++CurrentState.ArrayIndex;
-					}
-					else
-					{
-						//Too many entries in  TArray
-						UE_LOG(LogSerialization, Verbose, TEXT("TArray %s has dimension of %d and trying to read element %d"), *CurrentState.Property->GetFName().ToString(), ArrayHelper.Num(), CurrentState.ArrayIndex);
-						Backend.SkipStructure();
-						continue;
-					}
+					ArrayHelper.ExpandForIndex(CurrentState.ArrayIndex);
+
+					NewState.Property = ArrayProperty->Inner;
+					NewState.Data = ArrayHelper.GetRawPtr(CurrentState.ArrayIndex);
+					NewState.ArrayIndex = 0;
+					++CurrentState.ArrayIndex;
 				}
 				// handle struct element inside optional
 				else if (FOptionalProperty* OptionalProperty = CastField<FOptionalProperty>(CurrentState.Property))
 				{
 					void* Data = OptionalProperty->ContainerPtrToValuePtr<void>(CurrentState.Data);
+
 					NewState.Property = OptionalProperty->GetValueProperty();
 					NewState.Data = OptionalProperty->MarkSetAndGetInitializedValuePointerToReplace(Data);
 					NewState.ArrayIndex = 0;
@@ -1010,24 +957,16 @@ bool FStructDeserializer::DeserializeElement(void* OutAddress, UStruct& OwnerInf
 					if (FStructProperty* SetStructProperty = CastField<FStructProperty>(MapProperty->ValueProp))
 					{
 						FScriptMapHelper MapHelper(MapProperty, MapProperty->ContainerPtrToValuePtr<void>(CurrentState.Data));
-						
-						//If a specific index is asked and it's not valid, skip the property
-						const int32 InternalIndex = MapHelper.FindInternalIndex(CurrentState.ArrayIndex);
-						if (InternalIndex != INDEX_NONE)
-						{
-							//We're skipping a level directly so CurrentState becomes the outer (set) and NewState the inner (Element prop)
-							CurrentState.Property = NewState.Property;
-							NewState.Property = SetStructProperty;
-							NewState.Data = MapHelper.GetValuePtr(InternalIndex);
-							NewState.ArrayIndex = 0;
-						}
-						else
-						{
-							//Index out of bound
-							UE_LOG(LogSerialization, Verbose, TEXT("TMap %s has dimension of %d and trying to read element %d"), *GetNameSafe(CurrentState.Property), MapHelper.Num(), CurrentState.ArrayIndex);
-							Backend.SkipStructure();
-							continue;
-						}
+
+						const int32 InternalIndex = ExpandForIndex(MapHelper, CurrentState.ArrayIndex);
+
+						// We're skipping a level directly so CurrentState becomes the outer (set) and NewState the inner (Element prop)
+						CurrentState.Property = MapProperty;
+						CurrentState.Data = MapProperty->ContainerPtrToValuePtr<void>(CurrentState.Data);
+
+						NewState.Property = SetStructProperty;
+						NewState.Data = MapHelper.GetValuePtr(InternalIndex);
+						NewState.ArrayIndex = 0;
 					}
 				}
 				//Handle Set property entry
@@ -1036,24 +975,16 @@ bool FStructDeserializer::DeserializeElement(void* OutAddress, UStruct& OwnerInf
 					if (FStructProperty* SetStructProperty = CastField<FStructProperty>(SetProperty->ElementProp))
 					{
 						FScriptSetHelper SetHelper(SetProperty, SetProperty->ContainerPtrToValuePtr<void>(CurrentState.Data));
-						const int32 InternalIndex = SetHelper.FindInternalIndex(CurrentState.ArrayIndex);
 
-						//If a specific index is asked and it's not valid, skip the property
-						if (InternalIndex != INDEX_NONE)
-						{
-							//We're skipping a level directly so CurrentState becomes the outer (set) and NewState the inner (Element prop)
-							CurrentState.Property = NewState.Property;
-							NewState.Property = SetProperty->ElementProp;
-							NewState.Data = SetHelper.GetElementPtr(InternalIndex);
-							NewState.ArrayIndex = 0;
-						}
-						else
-						{
-							//Index out of bound
-							UE_LOG(LogSerialization, Verbose, TEXT("TArray %s has dimension of %d and trying to read element %d"), *GetNameSafe(CurrentState.Property), SetHelper.Num(), CurrentState.ArrayIndex);
-							Backend.SkipStructure();
-							continue;
-						}
+						const int32 InternalIndex = ExpandForIndex(SetHelper, CurrentState.ArrayIndex);
+
+						//We're skipping a level directly so CurrentState becomes the outer (set) and NewState the inner (Element prop)
+						CurrentState.Property = SetProperty;
+						CurrentState.Data = SetProperty->ContainerPtrToValuePtr<void>(CurrentState.Data);
+
+						NewState.Property = SetProperty->ElementProp;
+						NewState.Data = SetHelper.GetElementPtr(InternalIndex);
+						NewState.ArrayIndex = 0;
 					}
 				}
 				//Handle array property entry
@@ -1063,23 +994,16 @@ bool FStructDeserializer::DeserializeElement(void* OutAddress, UStruct& OwnerInf
 					{
 						FScriptArrayHelper ArrayHelper(ArrayProperty, ArrayProperty->ContainerPtrToValuePtr<void>(CurrentState.Data));
 
-						//If a specific index is asked and it's not valid, skip the property
-						if (ArrayHelper.IsValidIndex(CurrentState.ArrayIndex))
-						{
-							//NewState will become the outer when going through the properties.
-							//When reading from Array, we expect to read from one level. When its a struct, it's not.
-							CurrentState.Property = NewState.Property;
-							NewState.Property = ArrayProperty->Inner;
-							NewState.Data = ArrayHelper.GetRawPtr(CurrentState.ArrayIndex);
-							NewState.ArrayIndex = 0;
-						}
-						else
-						{
-							//Index out of bound
-							UE_LOG(LogSerialization, Verbose, TEXT("TArray %s has dimension of %d and trying to read element %d"), *GetNameSafe(CurrentState.Property), ArrayHelper.Num(), CurrentState.ArrayIndex);
-							Backend.SkipStructure();
-							continue;
-						}
+						ArrayHelper.ExpandForIndex(CurrentState.ArrayIndex);
+
+						//NewState will become the outer when going through the properties.
+						//When reading from Array, we expect to read from one level. When its a struct, it's not.
+						CurrentState.Property = ArrayProperty;
+						CurrentState.Data = ArrayProperty->ContainerPtrToValuePtr<void>(CurrentState.Data);
+
+						NewState.Property = ArrayProperty->Inner;
+						NewState.Data = ArrayHelper.GetRawPtr(CurrentState.ArrayIndex);
+						NewState.ArrayIndex = 0;
 					}
 				}
 				// handle optional property entry
diff --git a/Engine/Source/Runtime/Serialization/Private/Tests/StructSerializerTest.cpp b/Engine/Source/Runtime/Serialization/Private/Tests/StructSerializerTest.cpp
--- a/Engine/Source/Runtime/Serialization/Private/Tests/StructSerializerTest.cpp
+++ b/Engine/Source/Runtime/Serialization/Private/Tests/StructSerializerTest.cpp
@@ -454,11 +454,6 @@ namespace StructSerializerTest
 				FStructSerializer::SerializeElement(&TestStruct, Property, INDEX_NONE, SerializerBackend, Policies);
 
 				FStructSerializerTestStruct TestStruct2(NoInit);
-				TestStruct2.Arrays.Int32Array.SetNumUninitialized(TestStruct.Arrays.Int32Array.Num());
-				TestStruct2.Arrays.ByteArray.SetNumUninitialized(TestStruct.Arrays.ByteArray.Num());
-				TestStruct2.Arrays.VectorArray.SetNumUninitialized(TestStruct.Arrays.VectorArray.Num());
-				TestStruct2.Arrays.StructArray.SetNumZeroed(TestStruct.Arrays.StructArray.Num()); 
-
 				CHECK_MESSAGE(TEXT("Deserialization must succeed"), FStructDeserializer::DeserializeElement(&TestStruct2, *FStructSerializerTestStruct::StaticStruct(), INDEX_NONE, DeserializerBackend, DeserializerPolicies));
 
 				ValidateArrays(TestStruct.Arrays, TestStruct2.Arrays);
@@ -479,6 +474,8 @@ namespace StructSerializerTest
 				FStructSerializer::SerializeElement(&TestStruct, Property, INDEX_NONE, SerializerBackend, Policies);
 
 				FStructSerializerTestStruct TestStruct2(NoInit);
+
+				// Target Maps start off with keys copied as the map serialization policy is array and only deals with values. Leaves keys empty.
 				CopyKeys(TestStruct2.Maps.IntToStr, TestStruct.Maps.IntToStr);
 				CopyKeys(TestStruct2.Maps.StrToStr, TestStruct.Maps.StrToStr);
 				CopyKeys(TestStruct2.Maps.StrToStruct, TestStruct.Maps.StrToStruct);
@@ -504,12 +501,6 @@ namespace StructSerializerTest
 				FStructSerializer::SerializeElement(&TestStruct, Property, INDEX_NONE, SerializerBackend, Policies);
 
 				FStructSerializerTestStruct TestStruct2(NoInit);
-				TestStruct2.Sets.IntSet = { -1, -2, -3, -4, -5 };
-				TestStruct2.Sets.IntSet.Remove(-1);
-				TestStruct2.Sets.IntSet.Remove(-2);
-				TestStruct2.Sets.NameSet = { TEXT("Pre1"), TEXT("Pre2") , TEXT("Pre3") };
-				TestStruct2.Sets.StrSet = { TEXT("Pre4"), TEXT("Pre5") , TEXT("Pre6"), TEXT("Pre7") };
-				TestStruct2.Sets.StructSet = { FStructSerializerBuiltinTestStruct(NoInit) };
 				CHECK_MESSAGE(TEXT("Deserialization must succeed"), FStructDeserializer::DeserializeElement(&TestStruct2, *FStructSerializerTestStruct::StaticStruct(), INDEX_NONE, DeserializerBackend, DeserializerPolicies));
 
 				ValidateSets(TestStruct.Sets, TestStruct2.Sets);
