Fix StructSerializerTest trying to instanciate an abstract class.

diff --git a/Engine/Source/Runtime/Serialization/Private/Tests/StructSerializerTest.cpp b/Engine/Source/Runtime/Serialization/Private/Tests/StructSerializerTest.cpp
--- a/Engine/Source/Runtime/Serialization/Private/Tests/StructSerializerTest.cpp
+++ b/Engine/Source/Runtime/Serialization/Private/Tests/StructSerializerTest.cpp
@@ -336,19 +336,19 @@ namespace StructSerializerTest
 		// serialization
 	
 		FStructSerializerTestStruct OriginalStruct;
-		UClass* PackageMapClass = UPackageMap::StaticClass();
-		UPackageMap* PackageMapObject = NewObject<UPackageMap>();
+		UClass* ObjectTestClass = UObjectTest::StaticClass();
+		UObjectTest* ObjectTestObject = NewObject<UObjectTest>();
 		// setup object tests
-		OriginalStruct.Objects.RawClass = PackageMapClass;
-		OriginalStruct.Objects.WrappedClass = PackageMapClass;
-		OriginalStruct.Objects.SubClass = PackageMapClass;
-		OriginalStruct.Objects.SoftClass = PackageMapClass;
-		OriginalStruct.Objects.RawObject = PackageMapObject;
-		OriginalStruct.Objects.WrappedObject = PackageMapObject;
-		OriginalStruct.Objects.WeakObject = PackageMapObject;
-		OriginalStruct.Objects.SoftObject = PackageMapObject;
-		OriginalStruct.Objects.ClassPath = PackageMapClass;
-		OriginalStruct.Objects.ObjectPath = PackageMapObject;
+		OriginalStruct.Objects.RawClass = ObjectTestClass;
+		OriginalStruct.Objects.WrappedClass = ObjectTestClass;
+		OriginalStruct.Objects.SubClass = ObjectTestClass;
+		OriginalStruct.Objects.SoftClass = ObjectTestClass;
+		OriginalStruct.Objects.RawObject = ObjectTestObject;
+		OriginalStruct.Objects.WrappedObject = ObjectTestObject;
+		OriginalStruct.Objects.WeakObject = ObjectTestObject;
+		OriginalStruct.Objects.SoftObject = ObjectTestObject;
+		OriginalStruct.Objects.ClassPath = ObjectTestClass;
+		OriginalStruct.Objects.ObjectPath = ObjectTestObject;
 
 		{
 			FStructSerializerPolicies Policies;
@@ -725,19 +725,19 @@ namespace StructSerializerTest
 		// serialization
 		FStructSerializerTestStruct TestStruct;
 
-		UClass* PackageMapClass = UPackageMap::StaticClass();
-		UPackageMap* PackageMapObject = NewObject<UPackageMap>();
+		UClass* ObjectTestClass = UObjectTest::StaticClass();
+		UObjectTest* ObjectTestObject = NewObject<UObjectTest>();
 		// setup object tests
-		TestStruct.Objects.RawClass = PackageMapClass;
-		TestStruct.Objects.WrappedClass = PackageMapClass;
-		TestStruct.Objects.SubClass = PackageMapClass;
-		TestStruct.Objects.SoftClass = PackageMapClass;
-		TestStruct.Objects.RawObject = PackageMapObject;
-		TestStruct.Objects.WrappedObject = PackageMapObject;
-		TestStruct.Objects.WeakObject = PackageMapObject;
-		TestStruct.Objects.SoftObject = PackageMapObject;
-		TestStruct.Objects.ClassPath = PackageMapClass;
-		TestStruct.Objects.ObjectPath = PackageMapObject;
+		TestStruct.Objects.RawClass = ObjectTestClass;
+		TestStruct.Objects.WrappedClass = ObjectTestClass;
+		TestStruct.Objects.SubClass = ObjectTestClass;
+		TestStruct.Objects.SoftClass = ObjectTestClass;
+		TestStruct.Objects.RawObject = ObjectTestObject;
+		TestStruct.Objects.WrappedObject = ObjectTestObject;
+		TestStruct.Objects.WeakObject = ObjectTestObject;
+		TestStruct.Objects.SoftObject = ObjectTestObject;
+		TestStruct.Objects.ClassPath = ObjectTestClass;
+		TestStruct.Objects.ObjectPath = ObjectTestObject;
 
 		{
 			FStructSerializer::Serialize(TestStruct, SerializerBackend);
diff --git a/Engine/Source/Runtime/Serialization/Private/Tests/StructSerializerTestTypes.h b/Engine/Source/Runtime/Serialization/Private/Tests/StructSerializerTestTypes.h
--- a/Engine/Source/Runtime/Serialization/Private/Tests/StructSerializerTestTypes.h
+++ b/Engine/Source/Runtime/Serialization/Private/Tests/StructSerializerTestTypes.h
@@ -10,7 +10,12 @@
 #include "StructSerializerTestTypes.generated.h"
 
 class UObject;
-class UPackageMap;
+
+UCLASS()
+class UObjectTest : public UObject
+{
+	GENERATED_BODY()
+};
 
 /**
  * Test structure for numeric properties.
@@ -141,10 +146,10 @@ struct FStructSerializerObjectTestStruct
 	TObjectPtr<UClass> WrappedClass;
 
 	UPROPERTY()
-	TSubclassOf<UPackageMap> SubClass;
+	TSubclassOf<UObjectTest> SubClass;
 
 	UPROPERTY()
-	TSoftClassPtr<UPackageMap> SoftClass;
+	TSoftClassPtr<UObjectTest> SoftClass;
 
 	UPROPERTY()
 	TObjectPtr<UObject> RawObject;
@@ -153,10 +158,10 @@ struct FStructSerializerObjectTestStruct
 	TObjectPtr<UObject> WrappedObject;
 
 	UPROPERTY()
-	TWeakObjectPtr<UPackageMap> WeakObject;
+	TWeakObjectPtr<UObjectTest> WeakObject;
 
 	UPROPERTY()
-	TSoftObjectPtr<UPackageMap> SoftObject;
+	TSoftObjectPtr<UObjectTest> SoftObject;
 
 	UPROPERTY()
 	FSoftClassPath ClassPath;
