Fixed Cbor for FUtf8StrProperty and FAnsiStrProperty.

diff --git a/Engine/Source/Runtime/Cbor/Private/CborWriter.cpp b/Engine/Source/Runtime/Cbor/Private/CborWriter.cpp
--- a/Engine/Source/Runtime/Cbor/Private/CborWriter.cpp
+++ b/Engine/Source/Runtime/Cbor/Private/CborWriter.cpp
@@ -1,6 +1,8 @@
 // Copyright Epic Games, Inc. All Rights Reserved.
 
 #include "CborWriter.h"
+#include "Containers/AnsiString.h"
+#include "Containers/Utf8String.h"
 
 FCborWriter::FCborWriter(FArchive* InStream, ECborEndianness InWriterEndianness)
 	: Stream(InStream)
@@ -125,6 +127,16 @@ void FCborWriter::WriteValue(const FString& Value)
 	WriteValue(FStringView(Value));
 }
 
+void FCborWriter::WriteValue(const FAnsiString& Value)
+{
+	WriteValue(FAnsiStringView(Value));
+}
+
+void FCborWriter::WriteValue(const FUtf8String& Value)
+{
+	WriteValue(FUtf8StringView(Value));
+}
+
 void FCborWriter::WriteValue(const FStringView& Value)
 {
 	ScopedCborArchiveEndianness ScopedArchiveEndianness(*Stream, Endianness);
@@ -138,6 +150,18 @@ void FCborWriter::WriteValue(const FStringView& Value)
 	Stream->Serialize(const_cast<char*>(reinterpret_cast<const char*>(UTF8String.Get())), UTF8String.Length());
 }
 
+void FCborWriter::WriteValue(const FUtf8StringView& Value)
+{
+	ScopedCborArchiveEndianness ScopedArchiveEndianness(*Stream, Endianness);
+
+	CheckContext(ECborCode::TextString);
+	// Write string header
+	WriteUIntValue(ECborCode::TextString, *Stream, (uint64)Value.Len());
+	// Write string
+	check(sizeof(decltype(*Value.GetData())) == 1);
+	Stream->Serialize(const_cast<char*>(reinterpret_cast<const char*>(Value.GetData())), Value.Len());
+}
+
 void FCborWriter::WriteValue(const char* CString, uint64 Length)
 {
 	ScopedCborArchiveEndianness ScopedArchiveEndianness(*Stream, Endianness);
diff --git a/Engine/Source/Runtime/Cbor/Public/CborTypes.h b/Engine/Source/Runtime/Cbor/Public/CborTypes.h
--- a/Engine/Source/Runtime/Cbor/Public/CborTypes.h
+++ b/Engine/Source/Runtime/Cbor/Public/CborTypes.h
@@ -3,6 +3,8 @@
 #pragma once
 
 #include "CoreMinimal.h"
+#include "Containers/AnsiString.h"
+#include "Containers/Utf8String.h"
 #include "Misc/EnumClassFlags.h"
 
 /** 
@@ -256,6 +258,21 @@ struct FCborContext
 		return FString(FUTF8ToTCHAR(RawTextValue.GetData()).Get());
 	}
 
+	/** @return the context as a string. */
+	FAnsiString AsAnsiString() const
+	{
+		// Data loss is possible here
+		check(MajorType() == ECborCode::TextString);
+		return FAnsiString(RawTextValue.GetData());
+	}
+
+	/** @return the context as a string. */
+	FUtf8String AsUtf8String() const
+	{
+		check(MajorType() == ECborCode::TextString);
+		return FUtf8String(RawTextValue.GetData());
+	}
+
 	/** @return the context as a C string. */
 	const char* AsCString() const
 	{
diff --git a/Engine/Source/Runtime/Cbor/Public/CborWriter.h b/Engine/Source/Runtime/Cbor/Public/CborWriter.h
--- a/Engine/Source/Runtime/Cbor/Public/CborWriter.h
+++ b/Engine/Source/Runtime/Cbor/Public/CborWriter.h
@@ -3,7 +3,9 @@
 #pragma once
 
 #include "CborTypes.h"
+#include "Containers/ContainersFwd.h"
 #include "Containers/Array.h"
+#include "Containers/StringFwd.h"
 #include "Containers/UnrealString.h"
 #include "CoreMinimal.h"
 #include "HAL/Platform.h"
@@ -49,7 +51,10 @@ public:
 	CBOR_API void WriteValue(float Value);
 	CBOR_API void WriteValue(double Value);
 	CBOR_API void WriteValue(const FString& Value);
+	CBOR_API void WriteValue(const FAnsiString& Value);
+	CBOR_API void WriteValue(const FUtf8String& Value);
 	CBOR_API void WriteValue(const FStringView& Value);
+	CBOR_API void WriteValue(const FUtf8StringView& Value);
 	CBOR_API void WriteValue(const char* CString, uint64 Length);
 	CBOR_API void WriteValue(const uint8* Bytes, uint64 Length);
 
diff --git a/Engine/Source/Runtime/Serialization/Private/Backends/CborStructDeserializerBackend.cpp b/Engine/Source/Runtime/Serialization/Private/Backends/CborStructDeserializerBackend.cpp
--- a/Engine/Source/Runtime/Serialization/Private/Backends/CborStructDeserializerBackend.cpp
+++ b/Engine/Source/Runtime/Serialization/Private/Backends/CborStructDeserializerBackend.cpp
@@ -5,9 +5,11 @@
 #include "Backends/StructDeserializerBackendUtilities.h"
 #include "StructSerializationUtilities.h"
 #include "UObject/Class.h"
+#include "UObject/AnsiStrProperty.h"
 #include "UObject/EnumProperty.h"
 #include "UObject/TextProperty.h"
 #include "UObject/UnrealType.h"
+#include "UObject/Utf8StrProperty.h"
 
 
 FCborStructDeserializerBackend::FCborStructDeserializerBackend(FArchive& Archive, ECborEndianness CborDataEndianness, bool bInIsLWCCompatibilityMode)
@@ -173,6 +175,16 @@ bool FCborStructDeserializerBackend::ReadProperty(FProperty* Property, FProperty
 	// Strings, Names, Enumerations & Object/Class reference
 	case ECborCode::TextString:
 	{
+		if (FAnsiStrProperty* AnsiStrProperty = CastField<FAnsiStrProperty>(Property))
+		{
+			return StructDeserializerBackendUtilities::SetPropertyValue(AnsiStrProperty, Outer, Data, ArrayIndex, LastContext.AsAnsiString());
+		}
+
+		if (FUtf8StrProperty* Utf8StrProperty = CastField<FUtf8StrProperty>(Property))
+		{
+			return StructDeserializerBackendUtilities::SetPropertyValue(Utf8StrProperty, Outer, Data, ArrayIndex, LastContext.AsUtf8String());
+		}
+
 		FString StringValue = LastContext.AsString();
 
 		if (FStrProperty* StrProperty = CastField<FStrProperty>(Property))
diff --git a/Engine/Source/Runtime/Serialization/Private/Backends/CborStructSerializerBackend.cpp b/Engine/Source/Runtime/Serialization/Private/Backends/CborStructSerializerBackend.cpp
--- a/Engine/Source/Runtime/Serialization/Private/Backends/CborStructSerializerBackend.cpp
+++ b/Engine/Source/Runtime/Serialization/Private/Backends/CborStructSerializerBackend.cpp
@@ -3,10 +3,12 @@
 #include "Backends/CborStructSerializerBackend.h"
 
 #include "StructSerializationUtilities.h"
+#include "UObject/AnsiStrProperty.h"
 #include "UObject/EnumProperty.h"
 #include "UObject/PropertyPortFlags.h"
 #include "UObject/TextProperty.h"
 #include "UObject/UnrealType.h"
+#include "UObject/Utf8StrProperty.h"
 
 FCborStructSerializerBackend::FCborStructSerializerBackend(FArchive& InArchive)
 	: CborWriter(&InArchive)
@@ -281,6 +283,14 @@ void FCborStructSerializerBackend::WriteProperty(const FStructSerializerState& S
 	{
 		WritePropertyValue(CborWriter, State, CastFieldChecked<FStrProperty>(State.ValueProperty)->GetPropertyValue_InContainer(State.ValueData, ArrayIndex));
 	}
+	else if (State.FieldType == FAnsiStrProperty::StaticClass())
+	{
+		WritePropertyValue(CborWriter, State, CastFieldChecked<FAnsiStrProperty>(State.ValueProperty)->GetPropertyValue_InContainer(State.ValueData, ArrayIndex));
+	}
+	else if (State.FieldType == FUtf8StrProperty::StaticClass())
+	{
+		WritePropertyValue(CborWriter, State, CastFieldChecked<FUtf8StrProperty>(State.ValueProperty)->GetPropertyValue_InContainer(State.ValueData, ArrayIndex));
+	}
 	else if (State.FieldType == FTextProperty::StaticClass())
 	{
 		const FText& TextValue = CastFieldChecked<FTextProperty>(State.ValueProperty)->GetPropertyValue_InContainer(State.ValueData, ArrayIndex);
