Remove EStructSerializerBackendFlags::Legacy

diff --git a/Engine/Plugins/Messaging/TcpMessaging/Source/TcpMessaging/Private/Transport/TcpSerializeMessageTask.cpp b/Engine/Plugins/Messaging/TcpMessaging/Source/TcpMessaging/Private/Transport/TcpSerializeMessageTask.cpp
--- a/Engine/Plugins/Messaging/TcpMessaging/Source/TcpMessaging/Private/Transport/TcpSerializeMessageTask.cpp
+++ b/Engine/Plugins/Messaging/TcpMessaging/Source/TcpMessaging/Private/Transport/TcpSerializeMessageTask.cpp
@@ -52,7 +52,7 @@ void FTcpSerializeMessageTask::DoTask(ENamedThreads::Type CurrentThread, const F
 
 		// serialize message body
 		UScriptStruct* MessageTypeInfoPtr = MessageContext->GetMessageTypeInfo().Get();
-		FJsonStructSerializerBackend Backend(Archive, EStructSerializerBackendFlags::Legacy);
+		FJsonStructSerializerBackend Backend(Archive, EStructSerializerBackendFlags::None);
 		FStructSerializer::Serialize(MessageContext->GetMessage(), *MessageTypeInfoPtr, Backend);
 
 		// enqueue to recipients
diff --git a/Engine/Plugins/Messaging/UdpMessaging/Source/UdpMessaging/Private/Transport/UdpSerializeMessageTask.cpp b/Engine/Plugins/Messaging/UdpMessaging/Source/UdpMessaging/Private/Transport/UdpSerializeMessageTask.cpp
--- a/Engine/Plugins/Messaging/UdpMessaging/Source/UdpMessaging/Private/Transport/UdpSerializeMessageTask.cpp
+++ b/Engine/Plugins/Messaging/UdpMessaging/Source/UdpMessaging/Private/Transport/UdpSerializeMessageTask.cpp
@@ -46,7 +46,7 @@ void SerializeMessageV10(FArchive& Archive, const TSharedRef<IMessageContext, ES
 	}
 
 	// serialize message body
-	FJsonStructSerializerBackend Backend(Archive, EStructSerializerBackendFlags::Legacy);
+	FJsonStructSerializerBackend Backend(Archive, EStructSerializerBackendFlags::None);
 	FStructSerializer::Serialize(MessageContext->GetMessage(), *MessageContext->GetMessageTypeInfo(), Backend);
 }
 
@@ -164,7 +164,7 @@ void FUdpSerializeMessageTask::DoTask(ENamedThreads::Type CurrentThread, const F
 
 			case 11:
 				ProtocolMaxSegmentSize = UDP_MESSAGING_SEGMENT_SIZE * (int64)UINT16_MAX;
-				UdpSerializeMessageTaskDetails::SerializeMessageV11_15(Archive, MessageContext, EUdpMessageFormat::CborPlatformEndianness, EStructSerializerBackendFlags::Legacy);
+				UdpSerializeMessageTaskDetails::SerializeMessageV11_15(Archive, MessageContext, EUdpMessageFormat::CborPlatformEndianness, EStructSerializerBackendFlags::None);
 				break;
 
 			case 12:
diff --git a/Engine/Plugins/Runtime/Metasound/Source/MetasoundFrontend/Private/MetasoundJsonBackend.h b/Engine/Plugins/Runtime/Metasound/Source/MetasoundFrontend/Private/MetasoundJsonBackend.h
--- a/Engine/Plugins/Runtime/Metasound/Source/MetasoundFrontend/Private/MetasoundJsonBackend.h
+++ b/Engine/Plugins/Runtime/Metasound/Source/MetasoundFrontend/Private/MetasoundJsonBackend.h
@@ -175,18 +175,6 @@ namespace Metasound
 	{
 	public:
 
-		/**
-		 * Creates and initializes a new legacy instance.
-		 * @note Deprecated, use the two-parameter constructor with EStructSerializerBackendFlags::Legacy if you need backwards compatibility with code compiled prior to 4.22.
-		 *
-		 * @param InArchive The archive to serialize into.
-		 */
-		UE_DEPRECATED(4.22, "Use the two-parameter constructor with EStructSerializerBackendFlags::Legacy only if you need backwards compatibility with code compiled prior to 4.22; otherwise use EStructSerializerBackendFlags::Default.")
-			TJsonStructSerializerBackend(FArchive& InArchive)
-			: JsonWriter(TJsonWriter<CharType, PrintPolicy>::Create(&InArchive))
-			, Flags(EStructSerializerBackendFlags::Legacy)
-		{ }
-
 		/**
 		 * Creates and initializes a new instance with the given flags.
 		 *
diff --git a/Engine/Source/Runtime/Serialization/Private/Backends/CborStructSerializerBackend.cpp b/Engine/Source/Runtime/Serialization/Private/Backends/CborStructSerializerBackend.cpp
--- a/Engine/Source/Runtime/Serialization/Private/Backends/CborStructSerializerBackend.cpp
+++ b/Engine/Source/Runtime/Serialization/Private/Backends/CborStructSerializerBackend.cpp
@@ -11,11 +11,6 @@
 #include "UObject/UnrealType.h"
 #include "UObject/Utf8StrProperty.h"
 
-FCborStructSerializerBackend::FCborStructSerializerBackend(FArchive& InArchive)
-	: CborWriter(&InArchive)
-	, Flags(EStructSerializerBackendFlags::Legacy)
-{}
-
 FCborStructSerializerBackend::FCborStructSerializerBackend(FArchive& InArchive, const EStructSerializerBackendFlags InFlags)
 	: CborWriter(&InArchive, EnumHasAnyFlags(InFlags, EStructSerializerBackendFlags::WriteCborStandardEndianness) ? ECborEndianness::StandardCompliant : ECborEndianness::Platform)
 	, Flags(InFlags)
diff --git a/Engine/Source/Runtime/Serialization/Private/Tests/StructSerializerTest.cpp b/Engine/Source/Runtime/Serialization/Private/Tests/StructSerializerTest.cpp
--- a/Engine/Source/Runtime/Serialization/Private/Tests/StructSerializerTest.cpp
+++ b/Engine/Source/Runtime/Serialization/Private/Tests/StructSerializerTest.cpp
@@ -995,12 +995,10 @@ TEST_CASE_NAMED(FStructSerializerCborByteArrayTest, "System::Core::Serialization
 
 	// Ensure TArray<uint8>/TArray<int8> CBOR serialization is backward compatible. (Serializer can write the old format and deserializer can read it)
 	{
-		static_assert((EStructSerializerBackendFlags::Legacy & EStructSerializerBackendFlags::WriteByteArrayAsByteStream) == EStructSerializerBackendFlags::None, "Test below expects 'EStructSerializerBackendFlags::Legacy' to not have 'EStructSerializerBackendFlags::WriteByteArrayAsByteStream'");
-		
 		// Serialize TArray<uint8>/TArray<int8> it they were prior 4.25. (CBOR array rather than CBOR byte string)
 		TArray<uint8> Buffer;
 		FMemoryWriter Writer(Buffer);
-		FCborStructSerializerBackend SerializerBackend(Writer, EStructSerializerBackendFlags::Legacy); // Legacy mode doesn't enable EStructSerializerBackendFlags::WriteByteArrayAsByteStream.
+		FCborStructSerializerBackend SerializerBackend(Writer, EStructSerializerBackendFlags::None);
 		FStructSerializerByteArray WrittenStruct;
 		FStructSerializer::Serialize(WrittenStruct, SerializerBackend);
 
diff --git a/Engine/Source/Runtime/Serialization/Public/Backends/CborStructSerializerBackend.h b/Engine/Source/Runtime/Serialization/Public/Backends/CborStructSerializerBackend.h
--- a/Engine/Source/Runtime/Serialization/Public/Backends/CborStructSerializerBackend.h
+++ b/Engine/Source/Runtime/Serialization/Public/Backends/CborStructSerializerBackend.h
@@ -19,15 +19,6 @@ class FCborStructSerializerBackend
 {
 public:
 
-	/**
-	 * Creates and initializes a new legacy instance.
-	 * @note Deprecated, use the two-parameter constructor with EStructSerializerBackendFlags::Legacy if you need backwards compatibility with code compiled prior to 4.22.
-	 *
-	 * @param InArchive The archive to serialize into.
-	 */
-	UE_DEPRECATED(4.22, "Use the two-parameter constructor with EStructSerializerBackendFlags::Legacy only if you need backwards compatibility with code compiled prior to 4.22; otherwise use EStructSerializerBackendFlags::Default.")
-	SERIALIZATION_API FCborStructSerializerBackend(FArchive& InArchive);
-
 	/**
 	 * Creates and initializes a new instance with the given flags.
 	 *
diff --git a/Engine/Source/Runtime/Serialization/Public/Backends/JsonStructSerializerBackend.h b/Engine/Source/Runtime/Serialization/Public/Backends/JsonStructSerializerBackend.h
--- a/Engine/Source/Runtime/Serialization/Public/Backends/JsonStructSerializerBackend.h
+++ b/Engine/Source/Runtime/Serialization/Public/Backends/JsonStructSerializerBackend.h
@@ -27,18 +27,6 @@ class FJsonStructSerializerBackend
 {
 public:
 
-	/**
-	 * Creates and initializes a new legacy instance.
-	 * @note Deprecated, use the two-parameter constructor with EStructSerializerBackendFlags::Legacy if you need backwards compatibility with code compiled prior to 4.22.
-	 *
-	 * @param InArchive The archive to serialize into.
-	 */
-	UE_DEPRECATED(4.22, "Use the two-parameter constructor with EStructSerializerBackendFlags::Legacy only if you need backwards compatibility with code compiled prior to 4.22; otherwise use EStructSerializerBackendFlags::Default.")
-	FJsonStructSerializerBackend( FArchive& InArchive )
-		: JsonWriter(TJsonWriter<UCS2CHAR>::Create(&InArchive))
-		, Flags(EStructSerializerBackendFlags::Legacy)
-	{ }
-
 	/**
 	 * Creates and initializes a new instance with the given flags.
 	 *
diff --git a/Engine/Source/Runtime/Serialization/Public/IStructSerializerBackend.h b/Engine/Source/Runtime/Serialization/Public/IStructSerializerBackend.h
--- a/Engine/Source/Runtime/Serialization/Public/IStructSerializerBackend.h
+++ b/Engine/Source/Runtime/Serialization/Public/IStructSerializerBackend.h
@@ -39,11 +39,6 @@ enum class EStructSerializerBackendFlags
 	 */
 	WriteLWCTypesAsFloats = 1 << 3,
 
-	/**
-	 * Legacy settings for backwards compatibility with code compiled prior to 4.22.
-	 */
-	Legacy = None,
-
 	/**
 	 * Legacy settings for backwards compatibility with code compiled for 4.25 up to UE5.
 	 */
diff --git a/Engine/Source/Runtime/WebBrowser/Private/MobileJS/MobileJSStructSerializerBackend.cpp b/Engine/Source/Runtime/WebBrowser/Private/MobileJS/MobileJSStructSerializerBackend.cpp
--- a/Engine/Source/Runtime/WebBrowser/Private/MobileJS/MobileJSStructSerializerBackend.cpp
+++ b/Engine/Source/Runtime/WebBrowser/Private/MobileJS/MobileJSStructSerializerBackend.cpp
@@ -51,7 +51,7 @@ FString FMobileJSStructSerializerBackend::ToString()
 }
 
 FMobileJSStructSerializerBackend::FMobileJSStructSerializerBackend(TSharedRef<class FMobileJSScripting> InScripting)
-	: FJsonStructSerializerBackend(Writer, EStructSerializerBackendFlags::Legacy)
+	: FJsonStructSerializerBackend(Writer, EStructSerializerBackendFlags::None)
 	, Scripting(InScripting)
 	, ReturnBuffer()
 	, Writer(ReturnBuffer)
